#summary 	Fixed up issues found when creating a new source control provider.
#ttp 331349 	Editor: Source Control: Some issues with creating new source control provider
#change 	Fixed up conflation of enum with bool in if() blocks when calling Execute().
#fix 	Fixed QueueStatusUpdate not checking for NULL source control state.

diff --git a/Engine/Source/Developer/AssetTools/Private/AssetTools.cpp b/Engine/Source/Developer/AssetTools/Private/AssetTools.cpp
--- a/Engine/Source/Developer/AssetTools/Private/AssetTools.cpp
+++ b/Engine/Source/Developer/AssetTools/Private/AssetTools.cpp
@@ -1324,7 +1324,7 @@ void FAssetTools::MigratePackages_ReportConfirmed(TArray<FString> ConfirmedPacka
 			if(ISourceControlModule::Get().IsEnabled())
 			{
 				ISourceControlProvider& SourceControlProvider = ISourceControlModule::Get().GetProvider();
-				if(!SourceControlProvider.Execute(ISourceControlOperation::Create<FMarkForAdd>(), SuccessfullyCopiedFiles))
+				if(SourceControlProvider.Execute(ISourceControlOperation::Create<FMarkForAdd>(), SuccessfullyCopiedFiles) == ECommandResult::Failed)
 				{
 					for(auto FileIt(SuccessfullyCopiedFiles.CreateConstIterator()); FileIt; FileIt++)
 					{
diff --git a/Engine/Source/Developer/CollectionManager/Private/Collection.cpp b/Engine/Source/Developer/CollectionManager/Private/Collection.cpp
--- a/Engine/Source/Developer/CollectionManager/Private/Collection.cpp
+++ b/Engine/Source/Developer/CollectionManager/Private/Collection.cpp
@@ -508,7 +508,7 @@ bool FCollection::CheckoutCollection(FText& OutError)
 	// If not at the head revision, sync up
 	if (SourceControlState.IsValid() && !SourceControlState->IsCurrent())
 	{
-		if ( !SourceControlProvider.Execute(ISourceControlOperation::Create<FSync>(), FilesToBeCheckedOut) )
+		if ( SourceControlProvider.Execute(ISourceControlOperation::Create<FSync>(), FilesToBeCheckedOut) == ECommandResult::Failed )
 		{
 			// Could not sync up with the head revision
 			OutError = LOCTEXT("Error_SCCSync", "Failed to sync collection to the head revision.");
@@ -749,7 +749,7 @@ bool FCollection::RevertCollection(FText& OutError)
 		return false;
 	}
 
-	if ( SourceControlProvider.Execute(ISourceControlOperation::Create<FRevert>(), AbsoluteFilename) )
+	if ( SourceControlProvider.Execute(ISourceControlOperation::Create<FRevert>(), AbsoluteFilename) == ECommandResult::Succeeded)
 	{
 		return true;
 	}
@@ -811,7 +811,7 @@ bool FCollection::DeleteFromSourceControl(FText& OutError)
 	// If not at the head revision, sync up
 	if (SourceControlState.IsValid() && !SourceControlState->IsCurrent())
 	{
-		if ( !SourceControlProvider.Execute(ISourceControlOperation::Create<FSync>(), AbsoluteFilename) )
+		if ( SourceControlProvider.Execute(ISourceControlOperation::Create<FSync>(), AbsoluteFilename) == ECommandResult::Failed)
 		{
 			// Could not sync up with the head revision
 			GWarn->EndSlowTask();
@@ -854,7 +854,7 @@ bool FCollection::DeleteFromSourceControl(FText& OutError)
 		}
 		else if(SourceControlState->CanCheckout())
 		{
-			if ( SourceControlProvider.Execute(ISourceControlOperation::Create<FDelete>(), AbsoluteFilename) )
+			if ( SourceControlProvider.Execute(ISourceControlOperation::Create<FDelete>(), AbsoluteFilename) == ECommandResult::Succeeded )
 			{
 				// Now check in the delete
 				const FText ChangelistDesc = FText::Format( LOCTEXT("CollectionDeletedDesc", "Deleted collection: {CollectionName}"), CollectionNameText );
diff --git a/Engine/Source/Developer/SettingsEditor/Private/Widgets/SSettingsEditor.cpp b/Engine/Source/Developer/SettingsEditor/Private/Widgets/SSettingsEditor.cpp
--- a/Engine/Source/Developer/SettingsEditor/Private/Widgets/SSettingsEditor.cpp
+++ b/Engine/Source/Developer/SettingsEditor/Private/Widgets/SSettingsEditor.cpp
@@ -314,7 +314,7 @@ bool SSettingsEditor::CheckOutDefaultConfigFile( )
 			{
 				if (false)
 				{
-					if (SourceControlProvider.Execute(ISourceControlOperation::Create<FSync>(), FilesToBeCheckedOut))
+					if (SourceControlProvider.Execute(ISourceControlOperation::Create<FSync>(), FilesToBeCheckedOut) == ECommandResult::Succeeded)
 					{
 						SettingsObject->ReloadConfig();
 					}
@@ -326,7 +326,7 @@ bool SSettingsEditor::CheckOutDefaultConfigFile( )
 			}
 			else if (SourceControlState->CanCheckout() || SourceControlState->IsCheckedOutOther())
 			{
-				if (!SourceControlProvider.Execute(ISourceControlOperation::Create<FCheckOut>(), FilesToBeCheckedOut))
+				if (SourceControlProvider.Execute(ISourceControlOperation::Create<FCheckOut>(), FilesToBeCheckedOut) == ECommandResult::Failed)
 				{
 					ErrorMessage = LOCTEXT("FailedToCheckOutConfigFileError", "Error: Failed to check out the configuration file.");
 				}
diff --git a/Engine/Source/Developer/SourceControl/Private/SourceControlModule.cpp b/Engine/Source/Developer/SourceControl/Private/SourceControlModule.cpp
--- a/Engine/Source/Developer/SourceControl/Private/SourceControlModule.cpp
+++ b/Engine/Source/Developer/SourceControl/Private/SourceControlModule.cpp
@@ -287,10 +287,13 @@ void FSourceControlModule::QueueStatusUpdate(const FString& InFilename)
 	if(IsEnabled())
 	{
 		TSharedPtr<ISourceControlState, ESPMode::ThreadSafe> SourceControlState = GetProvider().GetState(InFilename, EStateCacheUsage::Use);
-		FTimespan TimeSinceLastUpdate = FDateTime::Now() - SourceControlState->GetTimeStamp();
-		if(TimeSinceLastUpdate > SourceControlConstants::StateRefreshInterval)
+		if(SourceControlState.IsValid())
 		{
-			PendingStatusUpdateFiles.AddUnique(InFilename);
+			FTimespan TimeSinceLastUpdate = FDateTime::Now() - SourceControlState->GetTimeStamp();
+			if(TimeSinceLastUpdate > SourceControlConstants::StateRefreshInterval)
+			{
+				PendingStatusUpdateFiles.AddUnique(InFilename);
+			}
 		}
 	}
 }
diff --git a/Engine/Source/Developer/SourceControl/Public/ISourceControlProvider.h b/Engine/Source/Developer/SourceControl/Public/ISourceControlProvider.h
--- a/Engine/Source/Developer/SourceControl/Public/ISourceControlProvider.h
+++ b/Engine/Source/Developer/SourceControl/Public/ISourceControlProvider.h
@@ -143,7 +143,7 @@ public:
 		TArray<FString> Files;
 		TArray< TSharedRef<ISourceControlState, ESPMode::ThreadSafe> > States;
 		Files.Add(InFile);
-		if(GetState(Files, States, InStateCacheUsage))
+		if(GetState(Files, States, InStateCacheUsage) == ECommandResult::Succeeded)
 		{
 			TSharedRef<ISourceControlState, ESPMode::ThreadSafe> State = States[0];
 			return State;
diff --git a/Engine/Source/Editor/DeviceProfileEditor/Private/SDeviceProfileEditor.cpp b/Engine/Source/Editor/DeviceProfileEditor/Private/SDeviceProfileEditor.cpp
--- a/Engine/Source/Editor/DeviceProfileEditor/Private/SDeviceProfileEditor.cpp
+++ b/Engine/Source/Editor/DeviceProfileEditor/Private/SDeviceProfileEditor.cpp
@@ -114,7 +114,7 @@ FReply SDeviceProfileSourceControl::HandleCheckoutButtonPressed()
 	{
 		TArray<FString> FilesToBeCheckedOut;
 		FilesToBeCheckedOut.Add(AbsoluteConfigFilePath);
-		if(!SourceControlProvider.Execute(ISourceControlOperation::Create<FCheckOut>(), FilesToBeCheckedOut))
+		if(SourceControlProvider.Execute(ISourceControlOperation::Create<FCheckOut>(), FilesToBeCheckedOut) == ECommandResult::Failed)
 		{
 			ErrorMessage = LOCTEXT("FailedToCheckOutConfigFileError", "Error: Failed to check out the configuration file.");
 		}
diff --git a/Engine/Source/Editor/UnrealEd/Private/Commandlets/UE4AudioRefactorMigrationCommandlet.cpp b/Engine/Source/Editor/UnrealEd/Private/Commandlets/UE4AudioRefactorMigrationCommandlet.cpp
--- a/Engine/Source/Editor/UnrealEd/Private/Commandlets/UE4AudioRefactorMigrationCommandlet.cpp
+++ b/Engine/Source/Editor/UnrealEd/Private/Commandlets/UE4AudioRefactorMigrationCommandlet.cpp
@@ -64,7 +64,7 @@ bool CheckoutPackage(UPackage* Package)
 		FSourceControlStatePtr SourceControlState = SourceControlProvider.GetState(PackageFilename, EStateCacheUsage::ForceUpdate);
 		if (SourceControlState.IsValid() && !SourceControlState->IsCheckedOut() && !SourceControlState->IsAdded())
 		{
-			if(!SourceControlProvider.Execute(ISourceControlOperation::Create<FCheckOut>(), Package))
+			if(SourceControlProvider.Execute(ISourceControlOperation::Create<FCheckOut>(), Package) == ECommandResult::Failed)
 			{
 				UE_LOG(LogUE4AudioRefactorMigrationCommandlet, Warning, TEXT("Could not check out package '%s'."), *PackageName);
 				return false;
diff --git a/Engine/Source/Editor/UnrealEd/Private/Editor.cpp b/Engine/Source/Editor/UnrealEd/Private/Editor.cpp
--- a/Engine/Source/Editor/UnrealEd/Private/Editor.cpp
+++ b/Engine/Source/Editor/UnrealEd/Private/Editor.cpp
@@ -6427,7 +6427,7 @@ void ExecuteInvalidateCachedShaders(const TArray< FString >& Args)
 	{
 		if( SourceControlState->CanCheckout() || SourceControlState->IsCheckedOutOther() )
 		{
-			if(!SourceControlProvider.Execute(ISourceControlOperation::Create<FCheckOut>(), FileName))
+			if(SourceControlProvider.Execute(ISourceControlOperation::Create<FCheckOut>(), FileName) == ECommandResult::Failed)
 			{
 				UE_LOG(LogConsoleResponse, Display, TEXT("r.InvalidateCachedShaders failed\nCouldn't check out \"ShaderVersion.usf\""));
 				return;
diff --git a/Engine/Source/Editor/UnrealEd/Private/ObjectTools.cpp b/Engine/Source/Editor/UnrealEd/Private/ObjectTools.cpp
--- a/Engine/Source/Editor/UnrealEd/Private/ObjectTools.cpp
+++ b/Engine/Source/Editor/UnrealEd/Private/ObjectTools.cpp
@@ -1416,7 +1416,7 @@ namespace ObjectTools
 					else
 					{
 						// Open the file for delete
-						if ( !SourceControlProvider.Execute(ISourceControlOperation::Create<FDelete>(), DeleteFilenames) )
+						if ( SourceControlProvider.Execute(ISourceControlOperation::Create<FDelete>(), DeleteFilenames) == ECommandResult::Failed )
 						{
 							UE_LOG(LogObjectTools, Warning, TEXT("SCC failed to open '%s' for delete while saving an empty package."), *PackageFilename);
 						}
