Collections code clean up

UETOOL-332 - Collections 2.0

Removed some duplicated data from FCollection (AssetList and AssetSet) as they had to be manually kept in sync. Only AssetSet exists now, and is converted into an array and sorted before being written to disk (to keep the order consistent for diffing). DiskAssetList has also been converted into a set (now DiskAssetSet) as this is only ever used for lookup queries.

All manually memory management has been removed from FCollectionManager, and the array of maps to collections has been converted into a single map using FCollectionNameType as its key. This simplifies collection lookup code in most cases, and removes the need to manually track the total number of collections outside of the map (NumCollections has now been removed).

Renamed ICollectionManager::GetCollectionsContainingAsset to ICollectionManager::GetCollectionsContainingObject as the function checks for any objects, not just assets.

Renamed FCollection::IsAssetInCollection to FCollection::IsObjectInCollection as the function checks for any objects, not just assets.

Made ICollectionManager::IsCollectionEmpty const.

Replaced all appropriate loop usage with range-based-for, and replaced any NULL with nullptr.

diff --git a/Engine/Source/Developer/CollectionManager/Private/Collection.cpp b/Engine/Source/Developer/CollectionManager/Private/Collection.cpp
--- a/Engine/Source/Developer/CollectionManager/Private/Collection.cpp
+++ b/Engine/Source/Developer/CollectionManager/Private/Collection.cpp
@@ -15,6 +15,25 @@ enum ECollectionFileVersion
 	COLLECTION_VER_AUTOMATIC_VERSION_PLUS_ONE
 };
 
+struct FCollectionUtils
+{
+	static void AppendCollectionToArray(const TSet<FName>& InAssetSet, TArray<FName>& OutAssetArray)
+	{
+		OutAssetArray.Reserve(OutAssetArray.Num() + InAssetSet.Num());
+		for (const FName& AssetName : InAssetSet)
+		{
+			OutAssetArray.Add(AssetName);
+		}
+	}
+
+	static TArray<FName> CreateCollectionArray(const TSet<FName>& InAssetSet)
+	{
+		TArray<FName> AssetArray;
+		AppendCollectionToArray(InAssetSet, AssetArray);
+		return AssetArray;
+	}
+};
+
 FCollection::FCollection()
 {
 
@@ -98,9 +117,10 @@ bool FCollection::LoadFromFile(const FString& InFilename, bool InUseSCC)
 	else
 	{
 		// Static collection, a flat list of asset paths
-		for (int32 LineIdx = 0; LineIdx < FileContents.Num(); ++LineIdx)
+		for (FString Line : FileContents)
 		{
-			const FString Line = FileContents[LineIdx].Trim().TrimTrailing();
+			Line.Trim();
+			Line.TrimTrailing();
 
 			if ( Line.Len() )
 			{
@@ -108,7 +128,7 @@ bool FCollection::LoadFromFile(const FString& InFilename, bool InUseSCC)
 			}
 		}
 
-		DiskAssetList = AssetList;
+		DiskAssetSet = AssetSet;
 	}
 
 	return true;
@@ -152,9 +172,9 @@ bool FCollection::Save(FText& OutError)
 	// Start with the header
 	TMap<FString,FString> HeaderPairs;
 	GenerateHeaderPairs(HeaderPairs);
-	for (TMap<FString,FString>::TConstIterator HeaderIt(HeaderPairs); HeaderIt; ++HeaderIt)
+	for (const auto& HeaderPair : HeaderPairs)
 	{
-		FileOutput += HeaderIt.Key() + TEXT(":") + HeaderIt.Value() + LINE_TERMINATOR;
+		FileOutput += HeaderPair.Key + TEXT(":") + HeaderPair.Value + LINE_TERMINATOR;
 	}
 	FileOutput += LINE_TERMINATOR;
 
@@ -165,10 +185,14 @@ bool FCollection::Save(FText& OutError)
 	}
 	else
 	{
+		// Write out the set as a sorted array to keep things in a known order for diffing
+		TArray<FName> AssetList = FCollectionUtils::CreateCollectionArray(AssetSet);
+		AssetList.Sort();
+
 		// Static collection. Save a flat list of all assets in the collection.
-		for (int32 AssetIdx = 0; AssetIdx < AssetList.Num(); ++AssetIdx)
+		for (const FName& AssetName : AssetList)
 		{
-			FileOutput += AssetList[AssetIdx].ToString() + LINE_TERMINATOR;
+			FileOutput += AssetName.ToString() + LINE_TERMINATOR;
 		}
 	}
 
@@ -227,7 +251,7 @@ bool FCollection::Save(FText& OutError)
 
 	if ( bSaveSuccessful )
 	{
-		DiskAssetList = AssetList;
+		DiskAssetSet = AssetSet;
 	}
 
 	GWarn->EndSlowTask();
@@ -264,18 +288,17 @@ bool FCollection::DeleteSourceFile(FText& OutError)
 
 	if ( bSuccessfullyDeleted )
 	{
-		DiskAssetList.Empty();
+		DiskAssetSet.Empty();
 	}
 
 	return bSuccessfullyDeleted;
 }
 
-bool FCollection::AddAssetToCollection (FName ObjectPath)
+bool FCollection::AddAssetToCollection(FName ObjectPath)
 {
 	// @todo collection Does this apply to dynamic collections?
 	if ( !AssetSet.Contains(ObjectPath) )
 	{
-		AssetList.Add(ObjectPath);
 		AssetSet.Add(ObjectPath);
 
 		return true;
@@ -284,21 +307,20 @@ bool FCollection::AddAssetToCollection (FName ObjectPath)
 	return false;
 }
 
-bool FCollection::RemoveAssetFromCollection (FName ObjectPath)
+bool FCollection::RemoveAssetFromCollection(FName ObjectPath)
 {
 	// @todo collection Does this apply to dynamic collections?
-	AssetSet.Remove(ObjectPath);
-	return AssetList.Remove(ObjectPath) > 0;
+	return AssetSet.Remove(ObjectPath) > 0;
 }
 
 void FCollection::GetAssetsInCollection(TArray<FName>& Assets) const
 {
 	// @todo collection Does this apply to dynamic collections?
-	for (int Index = 0; Index < AssetList.Num(); Index++)
+	for (const FName& AssetName : AssetSet)
 	{
-		if ( !AssetList[ Index ].ToString().StartsWith( TEXT("/Script/") ) )
+		if (!AssetName.ToString().StartsWith(TEXT("/Script/")))
 		{
-			Assets.Add( AssetList[ Index ] );
+			Assets.Add(AssetName);
 		}
 	}
 }
@@ -306,11 +328,11 @@ void FCollection::GetAssetsInCollection(TArray<FName>& Assets) const
 void FCollection::GetClassesInCollection(TArray<FName>& Classes) const
 {
 	// @todo collection Does this apply to dynamic collections?
-	for (int Index = 0; Index < AssetList.Num(); Index++)
+	for (const FName& AssetName : AssetSet)
 	{
-		if ( AssetList[ Index ].ToString().StartsWith( TEXT("/Script/") ) )
+		if (AssetName.ToString().StartsWith(TEXT("/Script/")))
 		{
-			Classes.Add( AssetList[ Index ] );
+			Classes.Add(AssetName);
 		}
 	}
 }
@@ -318,10 +340,10 @@ void FCollection::GetClassesInCollection(TArray<FName>& Classes) const
 void FCollection::GetObjectsInCollection(TArray<FName>& Objects) const
 {
 	// @todo collection Does this apply to dynamic collections?
-	Objects.Append(AssetList);
+	FCollectionUtils::AppendCollectionToArray(AssetSet, Objects);
 }
 
-bool FCollection::IsAssetInCollection(FName ObjectPath) const
+bool FCollection::IsObjectInCollection(FName ObjectPath) const
 {
 	// @todo collection Does this apply to dynamic collections?
 	return AssetSet.Contains(ObjectPath);
@@ -332,9 +354,8 @@ void FCollection::Clear()
 	CollectionName = NAME_None;
 	bUseSCC = false;
 	SourceFilename = TEXT("");
-	AssetList.Empty();
 	AssetSet.Empty();
-	DiskAssetList.Empty();
+	DiskAssetSet.Empty();
 }
 
 bool FCollection::IsDynamic() const
@@ -345,7 +366,7 @@ bool FCollection::IsDynamic() const
 	
 bool FCollection::IsEmpty() const
 {
-	return AssetList.Num() == 0;
+	return AssetSet.Num() == 0;
 }
 
 void FCollection::PrintCollection() const
@@ -359,9 +380,13 @@ void FCollection::PrintCollection() const
 		UE_LOG(LogCollectionManager, Log, TEXT("    Printing static elements of collection %s"), *CollectionName.ToString());
 		UE_LOG(LogCollectionManager, Log, TEXT("    ============================="));
 
-		for (int32 AssetIdx = 0; AssetIdx < AssetList.Num(); ++AssetIdx)
+		// Print the set as a sorted array to keep things in a sane order
+		TArray<FName> AssetList = FCollectionUtils::CreateCollectionArray(AssetSet);
+		AssetList.Sort();
+
+		for (const FName& AssetName : AssetList)
 		{
-			UE_LOG(LogCollectionManager, Log, TEXT("        %s"), *AssetList[AssetIdx].ToString());
+			UE_LOG(LogCollectionManager, Log, TEXT("        %s"), *AssetName.ToString());
 		}
 	}
 }
@@ -416,26 +441,19 @@ void FCollection::MergeWithCollection(const FCollection& Other)
 		GetDifferencesFromDisk(AssetsAdded, AssetsRemoved);
 
 		// Copy asset list from other collection
-		DiskAssetList = Other.DiskAssetList;
-		AssetList = DiskAssetList;
+		DiskAssetSet = Other.DiskAssetSet;
+		AssetSet = DiskAssetSet;
 
 		// Add the assets that were added before the merge
-		for (int32 AssetIdx = 0; AssetIdx < AssetsAdded.Num(); ++AssetIdx)
+		for (const FName& AddedAssetName : AssetsAdded)
 		{
-			AssetList.AddUnique(AssetsAdded[AssetIdx]);
+			AssetSet.Add(AddedAssetName);
 		}
 
 		// Remove the assets that were removed before the merge
-		for (int32 AssetIdx = 0; AssetIdx < AssetsRemoved.Num(); ++AssetIdx)
+		for (const FName& RemovedAssetName : AssetsRemoved)
 		{
-			AssetList.Remove(AssetsRemoved[AssetIdx]);
-		}
-
-		// Update the set
-		AssetSet.Empty();
-		for (int32 AssetIdx = 0; AssetIdx < AssetList.Num(); ++AssetIdx)
-		{
-			AssetSet.Add(AssetList[AssetIdx]);
+			AssetSet.Remove(RemovedAssetName);
 		}
 	}
 }
@@ -443,24 +461,20 @@ void FCollection::MergeWithCollection(const FCollection& Other)
 void FCollection::GetDifferencesFromDisk(TArray<FName>& AssetsAdded, TArray<FName>& AssetsRemoved)
 {
 	// Find the assets that were removed since the disk list
-	for (int32 AssetIdx = 0; AssetIdx < DiskAssetList.Num(); ++AssetIdx)
+	for (const FName& DiskAssetName : DiskAssetSet)
 	{
-		const FName& DiskAsset = DiskAssetList[AssetIdx];
-
-		if ( !AssetSet.Contains(DiskAsset) )
+		if (!AssetSet.Contains(DiskAssetName))
 		{
-			AssetsRemoved.Add(DiskAsset);
+			AssetsRemoved.Add(DiskAssetName);
 		}
 	}
 
 	// Find the assets that were added since the disk list
-	for (int32 AssetIdx = 0; AssetIdx < AssetList.Num(); ++AssetIdx)
+	for (const FName& MemoryAssetName : AssetSet)
 	{
-		const FName& MemoryAsset = AssetList[AssetIdx];
-
-		if ( !DiskAssetList.Contains(MemoryAsset) )
+		if (!DiskAssetSet.Contains(MemoryAssetName))
 		{
-			AssetsAdded.Add(MemoryAsset);
+			AssetsAdded.Add(MemoryAssetName);
 		}
 	}
 }
@@ -636,6 +650,9 @@ bool FCollection::CheckinCollection(FText& OutError)
 		TArray<FName> AssetsRemoved;
 		GetDifferencesFromDisk(AssetsAdded, AssetsRemoved);
 
+		AssetsAdded.Sort();
+		AssetsRemoved.Sort();
+
 		// Clear description
 		ChangelistDesc = FText::GetEmpty();
 
diff --git a/Engine/Source/Developer/CollectionManager/Private/Collection.h b/Engine/Source/Developer/CollectionManager/Private/Collection.h
--- a/Engine/Source/Developer/CollectionManager/Private/Collection.h
+++ b/Engine/Source/Developer/CollectionManager/Private/Collection.h
@@ -26,8 +26,8 @@ public:
 	void GetClassesInCollection(TArray<FName>& Classes) const;
 	/** Gets a list of objects in the collection. Static collections only. */
 	void GetObjectsInCollection(TArray<FName>& Objects) const;
-	/** Returns true when the specified asset is in the collection. Static collections only. */
-	bool IsAssetInCollection(FName ObjectPath) const;
+	/** Returns true when the specified object is in the collection. Static collections only. */
+	bool IsObjectInCollection(FName ObjectPath) const;
 	/** Resets a collection to its default values */
 	void Clear();
 
@@ -78,15 +78,12 @@ private:
 	/** The filename used to load this collection. Empty if it is new or never loaded from disk. */
 	FString SourceFilename;
 
-	/** A list of assets in the collection. Takes the form PackageName.AssetName */
-	TArray<FName> AssetList;
-
-	/** A set of assets in the collection. This mirrors AssetList entirely and is just used for fast lookups */
+	/** The set of assets in the collection. Takes the form PackageName.AssetName */
 	TSet<FName> AssetSet;
 
-	/** The list of assets that were in the collection last time it was loaded from or saved to disk */
-	TArray<FName> DiskAssetList;
+	/** The set of assets that were in the collection last time it was loaded from or saved to disk. Takes the form PackageName.AssetName */
+	TSet<FName> DiskAssetSet;
 
 	/** The file version for this collection */
 	int32 FileVersion;
-};
\ No newline at end of file
+};
diff --git a/Engine/Source/Developer/CollectionManager/Private/CollectionManager.cpp b/Engine/Source/Developer/CollectionManager/Private/CollectionManager.cpp
--- a/Engine/Source/Developer/CollectionManager/Private/CollectionManager.cpp
+++ b/Engine/Source/Developer/CollectionManager/Private/CollectionManager.cpp
@@ -6,7 +6,6 @@
 
 FCollectionManager::FCollectionManager()
 {
-	NumCollections = 0;
 	LastError = LOCTEXT("Error_Unknown", "None");
 
 	CollectionFolders[ECollectionShareType::CST_Local] = FPaths::GameSavedDir() / TEXT("Collections");
@@ -20,51 +19,28 @@ FCollectionManager::FCollectionManager()
 
 FCollectionManager::~FCollectionManager()
 {
-	// Delete all collections in the cache
-	for (int32 CacheIdx = 0; CacheIdx < ECollectionShareType::CST_All; ++CacheIdx)
-	{
-		for (TMap<FName, FCollection*>::TConstIterator CollectionIt(CachedCollections[CacheIdx]); CollectionIt; ++CollectionIt)
-		{
-			if ( CollectionIt.Value() )
-			{
-				delete CollectionIt.Value();
-				NumCollections--;
-			}
-		}
-
-		CachedCollections[CacheIdx].Empty();
-	}
-
-	// Make sure we have deleted all our allocated FCollection objects
-	ensure(NumCollections == 0);
 }
 
 void FCollectionManager::GetCollectionNames(ECollectionShareType::Type ShareType, TArray<FName>& CollectionNames) const
 {
-	if ( ShareType == ECollectionShareType::CST_All )
+	for (const auto& CachedCollection : CachedCollections)
 	{
-		// Asked for all share types, get names for each cache
-		for (int32 CacheIdx = 0; CacheIdx < ECollectionShareType::CST_All; ++CacheIdx)
+		const FCollectionNameType& CollectionKey = CachedCollection.Key;
+		if (ShareType == ECollectionShareType::CST_All || ShareType == CollectionKey.Type)
 		{
-			TArray<FName> Names;
-			GetCollectionNames(ECollectionShareType::Type(CacheIdx), Names);
-			CollectionNames.Append(Names);
+			CollectionNames.Add(CollectionKey.Name);
 		}
 	}
-	else
-	{
-		CachedCollections[ShareType].GenerateKeyArray(CollectionNames);
-	}
 }
 
 bool FCollectionManager::CollectionExists(FName CollectionName, ECollectionShareType::Type ShareType) const
 {
-	if ( ShareType == ECollectionShareType::CST_All )
+	if (ShareType == ECollectionShareType::CST_All)
 	{
-		// Asked for all share types, get names for each cache
+		// Asked to check all share types...
 		for (int32 CacheIdx = 0; CacheIdx < ECollectionShareType::CST_All; ++CacheIdx)
 		{
-			if ( CollectionExists(CollectionName, ECollectionShareType::Type(CacheIdx)) )
+			if (CachedCollections.Contains(FCollectionNameType(CollectionName, ECollectionShareType::Type(CacheIdx))))
 			{
 				// Collection exists in at least one cache
 				return true;
@@ -76,151 +52,113 @@ bool FCollectionManager::CollectionExists(FName CollectionName, ECollectionShare
 	}
 	else
 	{
-		return CachedCollections[ShareType].Contains(CollectionName);
+		return CachedCollections.Contains(FCollectionNameType(CollectionName, ShareType));
 	}
 }
 
 bool FCollectionManager::GetAssetsInCollection(FName CollectionName, ECollectionShareType::Type ShareType, TArray<FName>& AssetsPaths) const
 {
-	if ( ShareType == ECollectionShareType::CST_All )
+	auto GetAssetsInCollectionInternal = [this, &AssetsPaths](const FCollectionNameType& InCollectionKey) -> bool
+	{
+		const TSharedRef<FCollection>* const CollectionRefPtr = CachedCollections.Find(InCollectionKey);
+		if (CollectionRefPtr)
+		{
+			(*CollectionRefPtr)->GetAssetsInCollection(AssetsPaths);
+			return true;
+		}
+		return false;
+	};
+
+	if (ShareType == ECollectionShareType::CST_All)
 	{
 		// Asked for all share types, find assets in the specified collection name in any cache
 		bool bFoundAssets = false;
-
 		for (int32 CacheIdx = 0; CacheIdx < ECollectionShareType::CST_All; ++CacheIdx)
 		{
-			if ( GetAssetsInCollection(CollectionName, ECollectionShareType::Type(CacheIdx), AssetsPaths) )
-			{
-				bFoundAssets = true;
-			}
+			bFoundAssets |= GetAssetsInCollectionInternal(FCollectionNameType(CollectionName, ECollectionShareType::Type(CacheIdx)));
 		}
-
 		return bFoundAssets;
 	}
 	else
 	{
-		FCollection*const* CollectionPtr = CachedCollections[ShareType].Find(CollectionName);
-		FCollection* Collection = NULL;
-	
-		if (CollectionPtr)
-		{
-			Collection = *CollectionPtr;
-			ensure(Collection);
-		}
+		return GetAssetsInCollectionInternal(FCollectionNameType(CollectionName, ShareType));
+	}
+}
 
-		if ( Collection )
+bool FCollectionManager::GetClassesInCollection(FName CollectionName, ECollectionShareType::Type ShareType, TArray<FName>& ClassPaths) const
+{
+	auto GetClassesInCollectionInternal = [this, &ClassPaths](const FCollectionNameType& InCollectionKey) -> bool
+	{
+		const TSharedRef<FCollection>* const CollectionRefPtr = CachedCollections.Find(InCollectionKey);
+		if (CollectionRefPtr)
 		{
-			Collection->GetAssetsInCollection(AssetsPaths);
+			(*CollectionRefPtr)->GetClassesInCollection(ClassPaths);
 			return true;
 		}
-
 		return false;
-	}
-}
+	};
 
-bool FCollectionManager::GetClassesInCollection(FName CollectionName, ECollectionShareType::Type ShareType, TArray<FName>& ClassPaths) const
-{
-	if ( ShareType == ECollectionShareType::CST_All )
+	if (ShareType == ECollectionShareType::CST_All)
 	{
-		// Asked for all share types, find assets in the specified collection name in any cache
+		// Asked for all share types, find classes in the specified collection name in any cache
 		bool bFoundAssets = false;
-
 		for (int32 CacheIdx = 0; CacheIdx < ECollectionShareType::CST_All; ++CacheIdx)
 		{
-			if ( GetClassesInCollection(CollectionName, ECollectionShareType::Type(CacheIdx), ClassPaths) )
-			{
-				bFoundAssets = true;
-			}
+			bFoundAssets |= GetClassesInCollectionInternal(FCollectionNameType(CollectionName, ECollectionShareType::Type(CacheIdx)));
 		}
-
 		return bFoundAssets;
 	}
 	else
 	{
-		FCollection*const* CollectionPtr = CachedCollections[ShareType].Find(CollectionName);
-		FCollection* Collection = NULL;
-
-		if (CollectionPtr)
-		{
-			Collection = *CollectionPtr;
-			ensure(Collection);
-		}
+		return GetClassesInCollectionInternal(FCollectionNameType(CollectionName, ShareType));
+	}
+}
 
-		if ( Collection )
+bool FCollectionManager::GetObjectsInCollection(FName CollectionName, ECollectionShareType::Type ShareType, TArray<FName>& ObjectPaths) const
+{
+	auto GetObjectsInCollectionInternal = [this, &ObjectPaths](const FCollectionNameType& InCollectionKey) -> bool
+	{
+		const TSharedRef<FCollection>* const CollectionRefPtr = CachedCollections.Find(InCollectionKey);
+		if (CollectionRefPtr)
 		{
-			Collection->GetClassesInCollection(ClassPaths);
+			(*CollectionRefPtr)->GetObjectsInCollection(ObjectPaths);
 			return true;
 		}
-
 		return false;
-	}
-}
+	};
 
-bool FCollectionManager::GetObjectsInCollection(FName CollectionName, ECollectionShareType::Type ShareType, TArray<FName>& ObjectPaths) const
-{
-	if ( ShareType == ECollectionShareType::CST_All )
+	if (ShareType == ECollectionShareType::CST_All)
 	{
-		// Asked for all share types, find assets in the specified collection name in any cache
+		// Asked for all share types, find objects in the specified collection name in any cache
 		bool bFoundAssets = false;
-
 		for (int32 CacheIdx = 0; CacheIdx < ECollectionShareType::CST_All; ++CacheIdx)
 		{
-			if ( GetObjectsInCollection(CollectionName, ECollectionShareType::Type(CacheIdx), ObjectPaths) )
-			{
-				bFoundAssets = true;
-			}
+			bFoundAssets |= GetObjectsInCollectionInternal(FCollectionNameType(CollectionName, ECollectionShareType::Type(CacheIdx)));
 		}
-
 		return bFoundAssets;
 	}
 	else
 	{
-		FCollection*const* CollectionPtr = CachedCollections[ShareType].Find(CollectionName);
-		FCollection* Collection = NULL;
-
-		if (CollectionPtr)
-		{
-			Collection = *CollectionPtr;
-			ensure(Collection);
-		}
-
-		if ( Collection )
-		{
-			Collection->GetObjectsInCollection(ObjectPaths);
-			return true;
-		}
-
-		return false;
+		return GetObjectsInCollectionInternal(FCollectionNameType(CollectionName, ShareType));
 	}
 }
 
-void FCollectionManager::GetCollectionsContainingAsset(FName ObjectPath, ECollectionShareType::Type ShareType, TArray<FName>& OutCollectionNames) const
+void FCollectionManager::GetCollectionsContainingObject(FName ObjectPath, ECollectionShareType::Type ShareType, TArray<FName>& OutCollectionNames) const
 {
-	if ( ShareType == ECollectionShareType::CST_All )
+	for (const auto& CachedCollection : CachedCollections)
 	{
-		// Asked for all share types, find collections containing the specified asset in any cache
-		for (int32 CacheIdx = 0; CacheIdx < ECollectionShareType::CST_All; ++CacheIdx)
+		const FCollectionNameType& CollectionKey = CachedCollection.Key;
+		const TSharedRef<FCollection>& Collection = CachedCollection.Value;
+		if ((ShareType == ECollectionShareType::CST_All || ShareType == CollectionKey.Type) && Collection->IsObjectInCollection(ObjectPath))
 		{
-			GetCollectionsContainingAsset(ObjectPath, ECollectionShareType::Type(CacheIdx), OutCollectionNames);
-		}
-	}
-	else
-	{
-		// Ask each collection if it contains the asset
-		for (TMap<FName,FCollection*>::TConstIterator CollectionIt(CachedCollections[ShareType]); CollectionIt; ++CollectionIt)
-		{
-			const FCollection* Collection = CollectionIt.Value();
-			if ( ensure(Collection) && Collection->IsAssetInCollection(ObjectPath) )
-			{
-				OutCollectionNames.Add(CollectionIt.Key());
-			}
+			OutCollectionNames.Add(CollectionKey.Name);
 		}
 	}
 }
 
 void FCollectionManager::CreateUniqueCollectionName(const FName& BaseName, ECollectionShareType::Type ShareType, FName& OutCollectionName) const
 {
-	if(!ensure(ShareType != ECollectionShareType::CST_All))
+	if (!ensure(ShareType != ECollectionShareType::CST_All))
 	{
 		return;
 	}
@@ -229,7 +167,7 @@ void FCollectionManager::CreateUniqueCollectionName(const FName& BaseName, EColl
 	bool CollectionAlreadyExists = false;
 	do
 	{
-		if(IntSuffix <= 1)
+		if (IntSuffix <= 1)
 		{
 			OutCollectionName = BaseName;
 		}
@@ -238,15 +176,15 @@ void FCollectionManager::CreateUniqueCollectionName(const FName& BaseName, EColl
 			OutCollectionName = *FString::Printf(TEXT("%s%d"), *BaseName.ToString(), IntSuffix);
 		}
 
-		CollectionAlreadyExists = CachedCollections[ShareType].Contains(OutCollectionName);
+		CollectionAlreadyExists = CachedCollections.Contains(FCollectionNameType(OutCollectionName, ShareType));
 		++IntSuffix;
 	}
-	while(CollectionAlreadyExists);
+	while (CollectionAlreadyExists);
 }
 
 bool FCollectionManager::CreateCollection(FName CollectionName, ECollectionShareType::Type ShareType)
 {
-	if ( !ensure(ShareType < ECollectionShareType::CST_All) )
+	if (!ensure(ShareType < ECollectionShareType::CST_All))
 	{
 		// Bad share type
 		LastError = LOCTEXT("Error_Internal", "There was an internal error.");
@@ -257,33 +195,31 @@ bool FCollectionManager::CreateCollection(FName CollectionName, ECollectionShare
 	bool bUseSCC = ShouldUseSCC(ShareType);
 	FString SourceFolder = CollectionFolders[ShareType];
 
-	FCollection NewCollection(CollectionName, SourceFolder, CollectionExtension, bUseSCC);
-	FCollection* Collection = AddCollection(NewCollection, ShareType);
-
-	if ( !Collection )
+	TSharedRef<FCollection> NewCollection = MakeShareable(new FCollection(CollectionName, SourceFolder, CollectionExtension, bUseSCC));
+	if (!AddCollection(NewCollection, ShareType))
 	{
 		// Failed to add the collection, it already exists
 		LastError = LOCTEXT("Error_AlreadyExists", "The collection already exists.");
 		return false;
 	}
 
-	if ( Collection->Save(LastError) )
+	if (NewCollection->Save(LastError))
 	{
 		// Collection saved!
-		CollectionCreatedEvent.Broadcast( FCollectionNameType( CollectionName, ShareType ) );
+		CollectionCreatedEvent.Broadcast(FCollectionNameType(CollectionName, ShareType));
 		return true;
 	}
 	else
 	{
 		// Collection failed to save, remove it from the cache
-		RemoveCollection(Collection, ShareType);
+		RemoveCollection(NewCollection, ShareType);
 		return false;
 	}
 }
 
-bool FCollectionManager::RenameCollection (FName CurrentCollectionName, ECollectionShareType::Type CurrentShareType, FName NewCollectionName, ECollectionShareType::Type NewShareType)
+bool FCollectionManager::RenameCollection(FName CurrentCollectionName, ECollectionShareType::Type CurrentShareType, FName NewCollectionName, ECollectionShareType::Type NewShareType)
 {
-	if ( !ensure(CurrentShareType < ECollectionShareType::CST_All) || !ensure(NewShareType < ECollectionShareType::CST_All) )
+	if (!ensure(CurrentShareType < ECollectionShareType::CST_All) || !ensure(NewShareType < ECollectionShareType::CST_All))
 	{
 		// Bad share type
 		LastError = LOCTEXT("Error_Internal", "There was an internal error.");
@@ -292,23 +228,23 @@ bool FCollectionManager::RenameCollection (FName CurrentCollectionName, ECollect
 
 	// Get the object paths for the assets in the collection
 	TArray<FName> ObjectPaths;
-	if ( !GetAssetsInCollection(CurrentCollectionName, CurrentShareType, ObjectPaths) )
+	if (!GetAssetsInCollection(CurrentCollectionName, CurrentShareType, ObjectPaths))
 	{
 		// Failed to get assets in the current collection
 		return false;
 	}
 
 	// Create a new collection
-	if ( !CreateCollection(NewCollectionName, NewShareType) )
+	if (!CreateCollection(NewCollectionName, NewShareType))
 	{
 		// Failed to create collection
 		return false;
 	}
 
-	if ( ObjectPaths.Num() > 0 )
+	if (ObjectPaths.Num() > 0)
 	{
 		// Add all the objects from the old collection to the new collection
-		if ( !AddToCollection(NewCollectionName, NewShareType, ObjectPaths) )
+		if (!AddToCollection(NewCollectionName, NewShareType, ObjectPaths))
 		{
 			// Failed to add paths to the new collection. Destroy the collection we created.
 			DestroyCollection(NewCollectionName, NewShareType);
@@ -317,7 +253,7 @@ bool FCollectionManager::RenameCollection (FName CurrentCollectionName, ECollect
 	}
 
 	// Delete the old collection
-	if ( !DestroyCollection(CurrentCollectionName, CurrentShareType) )
+	if (!DestroyCollection(CurrentCollectionName, CurrentShareType))
 	{
 		// Failed to destroy the old collection. Destroy the collection we created.
 		DestroyCollection(NewCollectionName, NewShareType);
@@ -325,47 +261,44 @@ bool FCollectionManager::RenameCollection (FName CurrentCollectionName, ECollect
 	}
 
 	// Success
-	FCollectionNameType OriginalCollectionNameType( CurrentCollectionName, CurrentShareType );
-	FCollectionNameType NewCollectionNameType( NewCollectionName, NewShareType );
-	CollectionRenamedEvent.Broadcast( OriginalCollectionNameType, NewCollectionNameType );
+	const FCollectionNameType OriginalCollectionKey(CurrentCollectionName, CurrentShareType);
+	const FCollectionNameType NewCollectionKey(NewCollectionName, NewShareType);
+	CollectionRenamedEvent.Broadcast(OriginalCollectionKey, NewCollectionKey);
 	return true;
 }
 
 bool FCollectionManager::DestroyCollection(FName CollectionName, ECollectionShareType::Type ShareType)
 {
-	if ( !ensure(ShareType < ECollectionShareType::CST_All) )
+	if (!ensure(ShareType < ECollectionShareType::CST_All))
 	{
 		// Bad share type
 		LastError = LOCTEXT("Error_Internal", "There was an internal error.");
 		return false;
 	}
 
-	FCollection** CollectionPtr = CachedCollections[ShareType].Find(CollectionName);
-	if (CollectionPtr != NULL && ensure(*CollectionPtr != NULL) )
+	const FCollectionNameType CollectionKey(CollectionName, ShareType);
+	TSharedRef<FCollection>* const CollectionRefPtr = CachedCollections.Find(CollectionKey);
+	if (!CollectionRefPtr)
 	{
-		FCollection* Collection = *CollectionPtr;
-		
-		if ( Collection->DeleteSourceFile(LastError) )
-		{
-			RemoveCollection(Collection, ShareType);
-			CollectionDestroyedEvent.Broadcast( FCollectionNameType( CollectionName, ShareType ) );
-			return true;
-		}
-		else
-		{
-			// Failed to delete the source file
-			return false;
-		}
+		// The collection doesn't exist
+		LastError = LOCTEXT("Error_DoesntExist", "The collection doesn't exist.");
+		return false;
+	}
+
+	if ((*CollectionRefPtr)->DeleteSourceFile(LastError))
+	{
+		RemoveCollection(*CollectionRefPtr, ShareType);
+		CollectionDestroyedEvent.Broadcast(CollectionKey);
+		return true;
 	}
 	else
 	{
-		// The collection doesn't exist
-		LastError = LOCTEXT("Error_DoesntExist", "The collection doesn't exist.");
+		// Failed to delete the source file
 		return false;
 	}
 }
 
-bool FCollectionManager::AddToCollection(FName CollectionName, ECollectionShareType::Type ShareType,  FName ObjectPath)
+bool FCollectionManager::AddToCollection(FName CollectionName, ECollectionShareType::Type ShareType, FName ObjectPath)
 {
 	TArray<FName> Paths;
 	Paths.Add(ObjectPath);
@@ -374,77 +307,63 @@ bool FCollectionManager::AddToCollection(FName CollectionName, ECollectionShareT
 
 bool FCollectionManager::AddToCollection(FName CollectionName, ECollectionShareType::Type ShareType, const TArray<FName>& ObjectPaths, int32* OutNumAdded)
 {
-	if ( OutNumAdded )
+	if (OutNumAdded)
 	{
 		*OutNumAdded = 0;
 	}
 
-	if ( !ensure(ShareType < ECollectionShareType::CST_All) )
+	if (!ensure(ShareType < ECollectionShareType::CST_All))
 	{
 		// Bad share type
 		LastError = LOCTEXT("Error_Internal", "There was an internal error.");
 		return false;
 	}
 
-	FCollection** CollectionPtr = CachedCollections[ShareType].Find(CollectionName);
-	FCollection* Collection = NULL;
-	if (CollectionPtr != NULL)
-	{
-		Collection = *CollectionPtr;
-	}
-	else
+	const FCollectionNameType CollectionKey(CollectionName, ShareType);
+	TSharedRef<FCollection>* const CollectionRefPtr = CachedCollections.Find(CollectionKey);
+	if (!CollectionRefPtr)
 	{
 		// Collection doesn't exist
 		LastError = LOCTEXT("Error_DoesntExist", "The collection doesn't exist.");
 		return false;
 	}
 
-	if ( ensure(Collection) )
+	int32 NumAdded = 0;
+	for (const FName& ObjectPath : ObjectPaths)
 	{
-		int32 NumAdded = 0;
-		for (int32 ObjectIdx = 0; ObjectIdx < ObjectPaths.Num(); ++ObjectIdx)
+		if ((*CollectionRefPtr)->AddAssetToCollection(ObjectPath))
 		{
-			if ( Collection->AddAssetToCollection(ObjectPaths[ObjectIdx]) )
-			{
-				NumAdded++;
-			}
+			NumAdded++;
 		}
+	}
 
-		if ( NumAdded > 0 )
+	if (NumAdded > 0)
+	{
+		if ((*CollectionRefPtr)->Save(LastError))
 		{
-			if ( Collection->Save(LastError) )
-			{
-				// Added and saved
-				if ( OutNumAdded )
-				{
-					*OutNumAdded = NumAdded;
-				}
-
-				FCollectionNameType CollectionNameType( Collection->GetCollectionName(), ShareType );
-				AssetsAddedEvent.Broadcast( CollectionNameType, ObjectPaths );
-				return true;
-			}
-			else
+			// Added and saved
+			if (OutNumAdded)
 			{
-				// Added but not saved, revert the add
-				for (int32 ObjectIdx = 0; ObjectIdx < ObjectPaths.Num(); ++ObjectIdx)
-				{
-					Collection->RemoveAssetFromCollection(ObjectPaths[ObjectIdx]);
-				}
-				return false;
+				*OutNumAdded = NumAdded;
 			}
+
+			AssetsAddedEvent.Broadcast(CollectionKey, ObjectPaths);
+			return true;
 		}
 		else
 		{
-			// Failed to add, all of the objects were already in the collection
-			LastError = LOCTEXT("Error_AlreadyInCollection", "All of the assets were already in the collection.");
+			// Added but not saved, revert the add
+			for (const FName& ObjectPath : ObjectPaths)
+			{
+				(*CollectionRefPtr)->RemoveAssetFromCollection(ObjectPath);
+			}
 			return false;
 		}
 	}
 	else
 	{
-		// Collection was never found. Perhaps it was null in the TMap
-		LastError = LOCTEXT("Error_DoesntExist", "The collection doesn't exist.");
+		// Failed to add, all of the objects were already in the collection
+		LastError = LOCTEXT("Error_AlreadyInCollection", "All of the assets were already in the collection.");
 		return false;
 	}
 }
@@ -458,74 +377,68 @@ bool FCollectionManager::RemoveFromCollection(FName CollectionName, ECollectionS
 
 bool FCollectionManager::RemoveFromCollection(FName CollectionName, ECollectionShareType::Type ShareType, const TArray<FName>& ObjectPaths, int32* OutNumRemoved)
 {
-	if ( OutNumRemoved )
+	if (OutNumRemoved)
 	{
 		*OutNumRemoved = 0;
 	}
 
-	if ( !ensure(ShareType < ECollectionShareType::CST_All) )
+	if (!ensure(ShareType < ECollectionShareType::CST_All))
 	{
 		// Bad share type
 		LastError = LOCTEXT("Error_Internal", "There was an internal error.");
 		return false;
 	}
 
-	FCollection** CollectionPtr = CachedCollections[ShareType].Find(CollectionName);
-	if (CollectionPtr != NULL)
+	const FCollectionNameType CollectionKey(CollectionName, ShareType);
+	TSharedRef<FCollection>* const CollectionRefPtr = CachedCollections.Find(CollectionKey);
+	if (!CollectionRefPtr)
 	{
-		FCollection* Collection = *CollectionPtr;
-		
-		TArray<FName> RemovedAssets;
-		for (int32 ObjectIdx = 0; ObjectIdx < ObjectPaths.Num(); ++ObjectIdx)
-		{
-			const FName& ObjectPath = ObjectPaths[ObjectIdx];
-
-			if ( Collection->RemoveAssetFromCollection(ObjectPath) )
-			{
-				RemovedAssets.Add(ObjectPath);
-			}
-		}
+		// Collection not found
+		LastError = LOCTEXT("Error_DoesntExist", "The collection doesn't exist.");
+		return false;
+	}
 
-		if ( RemovedAssets.Num() == 0 )
+	TArray<FName> RemovedAssets;
+	for (const FName& ObjectPath : ObjectPaths)
+	{
+		if ((*CollectionRefPtr)->RemoveAssetFromCollection(ObjectPath))
 		{
-			// Failed to remove, none of the objects were in the collection
-			LastError = LOCTEXT("Error_NotInCollection", "None of the assets were in the collection.");
-			return false;
+			RemovedAssets.Add(ObjectPath);
 		}
-			
-		if ( Collection->Save(LastError) )
-		{
-			// Removed and saved
-			if ( OutNumRemoved )
-			{
-				*OutNumRemoved = RemovedAssets.Num();
-			}
+	}
 
-			FCollectionNameType CollectionNameType( Collection->GetCollectionName(), ShareType );
-			AssetsRemovedEvent.Broadcast( CollectionNameType, ObjectPaths );
-			return true;
-		}
-		else
+	if (RemovedAssets.Num() == 0)
+	{
+		// Failed to remove, none of the objects were in the collection
+		LastError = LOCTEXT("Error_NotInCollection", "None of the assets were in the collection.");
+		return false;
+	}
+			
+	if ((*CollectionRefPtr)->Save(LastError))
+	{
+		// Removed and saved
+		if (OutNumRemoved)
 		{
-			// Removed but not saved, revert the remove
-			for (int32 AssetIdx = 0; AssetIdx < RemovedAssets.Num(); ++AssetIdx)
-			{
-				Collection->AddAssetToCollection(RemovedAssets[AssetIdx]);
-			}
-			return false;
+			*OutNumRemoved = RemovedAssets.Num();
 		}
+
+		AssetsRemovedEvent.Broadcast(CollectionKey, ObjectPaths);
+		return true;
 	}
 	else
 	{
-		// Collection not found
-		LastError = LOCTEXT("Error_DoesntExist", "The collection doesn't exist.");
+		// Removed but not saved, revert the remove
+		for (const FName& RemovedAssetName : RemovedAssets)
+		{
+			(*CollectionRefPtr)->AddAssetToCollection(RemovedAssetName);
+		}
 		return false;
 	}
 }
 
 bool FCollectionManager::EmptyCollection(FName CollectionName, ECollectionShareType::Type ShareType)
 {
-	if ( !ensure(ShareType < ECollectionShareType::CST_All) )
+	if (!ensure(ShareType < ECollectionShareType::CST_All))
 	{
 		// Bad share type
 		LastError = LOCTEXT("Error_Internal", "There was an internal error.");
@@ -533,13 +446,13 @@ bool FCollectionManager::EmptyCollection(FName CollectionName, ECollectionShareT
 	}
 
 	TArray<FName> ObjectPaths;
-	if ( !GetAssetsInCollection(CollectionName, ShareType, ObjectPaths) )
+	if (!GetAssetsInCollection(CollectionName, ShareType, ObjectPaths))
 	{
 		// Failed to load collection
 		return false;
 	}
 	
-	if ( ObjectPaths.Num() == 0 )
+	if (ObjectPaths.Num() == 0)
 	{
 		// Collection already empty
 		return true;
@@ -548,27 +461,20 @@ bool FCollectionManager::EmptyCollection(FName CollectionName, ECollectionShareT
 	return RemoveFromCollection(CollectionName, ShareType, ObjectPaths);
 }
 
-bool FCollectionManager::IsCollectionEmpty(FName CollectionName, ECollectionShareType::Type ShareType)
+bool FCollectionManager::IsCollectionEmpty(FName CollectionName, ECollectionShareType::Type ShareType) const
 {
-	if ( !ensure(ShareType < ECollectionShareType::CST_All) )
+	if (!ensure(ShareType < ECollectionShareType::CST_All))
 	{
 		// Bad share type
 		LastError = LOCTEXT("Error_Internal", "There was an internal error.");
 		return true;
 	}
 
-	FCollection*const* CollectionPtr = CachedCollections[ShareType].Find(CollectionName);
-	FCollection* Collection = NULL;
-
-	if ( CollectionPtr )
-	{
-		Collection = *CollectionPtr;
-		ensure(Collection);
-	}
-
-	if ( Collection )
+	const FCollectionNameType CollectionKey(CollectionName, ShareType);
+	const TSharedRef<FCollection>* const CollectionRefPtr = CachedCollections.Find(CollectionKey);
+	if (CollectionRefPtr)
 	{
-		return Collection->IsEmpty();
+		return (*CollectionRefPtr)->IsEmpty();
 	}
 
 	return true;
@@ -577,7 +483,7 @@ bool FCollectionManager::IsCollectionEmpty(FName CollectionName, ECollectionShar
 void FCollectionManager::LoadCollections()
 {
 	const double LoadStartTime = FPlatformTime::Seconds();
-	int32 NumCollectionsLoaded = 0;
+	const int32 PrevNumCollections = CachedCollections.Num();
 
 	for (int32 CacheIdx = 0; CacheIdx < ECollectionShareType::CST_All; ++CacheIdx)
 	{
@@ -587,26 +493,24 @@ void FCollectionManager::LoadCollections()
 		TArray<FString> Filenames;
 		IFileManager::Get().FindFiles(Filenames, *WildCard, true, false);
 
-		for (int32 FileIdx = 0; FileIdx < Filenames.Num(); ++FileIdx)
+		for (const FString& BaseFilename : Filenames)
 		{
-			const FString Filename = CollectionFolder / Filenames[FileIdx];
+			const FString Filename = CollectionFolder / BaseFilename;
 
-			FCollection Collection;
+			TSharedRef<FCollection> NewCollection = MakeShareable(new FCollection());
 			const bool bUseSCC = ShouldUseSCC(ECollectionShareType::Type(CacheIdx));
-			if ( Collection.LoadFromFile(Filename, bUseSCC) )
+			if (NewCollection->LoadFromFile(Filename, bUseSCC))
 			{
-				AddCollection(Collection, ECollectionShareType::Type(CacheIdx));
+				AddCollection(NewCollection, ECollectionShareType::Type(CacheIdx));
 			}
 			else
 			{
 				UE_LOG(LogCollectionManager, Warning, TEXT("Failed to load collection file %s"), *Filename);
 			}
 		}
-		
-		NumCollectionsLoaded += CachedCollections[CacheIdx].Num();
 	}
 
-	UE_LOG( LogCollectionManager, Log, TEXT( "Loaded %d collections in %0.6f seconds" ), NumCollectionsLoaded, FPlatformTime::Seconds() - LoadStartTime );
+	UE_LOG(LogCollectionManager, Log, TEXT( "Loaded %d collections in %0.6f seconds" ), CachedCollections.Num() - PrevNumCollections, FPlatformTime::Seconds() - LoadStartTime);
 }
 
 bool FCollectionManager::ShouldUseSCC(ECollectionShareType::Type ShareType) const
@@ -614,51 +518,35 @@ bool FCollectionManager::ShouldUseSCC(ECollectionShareType::Type ShareType) cons
 	return ShareType != ECollectionShareType::CST_Local && ShareType != ECollectionShareType::CST_System;
 }
 
-FCollection* FCollectionManager::AddCollection(const FCollection& Collection, ECollectionShareType::Type ShareType)
+bool FCollectionManager::AddCollection(const TSharedRef<FCollection>& CollectionRef, ECollectionShareType::Type ShareType)
 {
-	if ( !ensure(ShareType < ECollectionShareType::CST_All) )
+	if (!ensure(ShareType < ECollectionShareType::CST_All))
 	{
 		// Bad share type
-		return NULL;
+		return false;
 	}
 
-	FName CollectionName = Collection.GetCollectionName();
-	if ( CachedCollections[ShareType].Contains(CollectionName) )
+	const FCollectionNameType CollectionKey(CollectionRef->GetCollectionName(), ShareType);
+	if (CachedCollections.Contains(CollectionKey))
 	{
-		UE_LOG(LogCollectionManager, Warning, TEXT("Failed to add collection '%s' because it already exists."), *CollectionName.ToString());
-		return NULL;
+		UE_LOG(LogCollectionManager, Warning, TEXT("Failed to add collection '%s' because it already exists."), *CollectionRef->GetCollectionName().ToString());
+		return false;
 	}
 
-	FCollection* NewCollection = new FCollection(Collection);
-	NumCollections++;
-
-	CachedCollections[ShareType].Add(CollectionName, NewCollection);
-
-	return NewCollection;
+	CachedCollections.Add(CollectionKey, CollectionRef);
+	return true;
 }
 
-bool FCollectionManager::RemoveCollection(FCollection* Collection, ECollectionShareType::Type ShareType)
+bool FCollectionManager::RemoveCollection(const TSharedRef<FCollection>& CollectionRef, ECollectionShareType::Type ShareType)
 {
-	if ( !ensure(ShareType < ECollectionShareType::CST_All) )
+	if (!ensure(ShareType < ECollectionShareType::CST_All))
 	{
 		// Bad share type
 		return false;
 	}
 
-	bool bRemoved = false;
-
-	if ( ensure(Collection && Collection->GetCollectionName() != NAME_None) )
-	{
-		bRemoved = CachedCollections[ShareType].Remove(Collection->GetCollectionName()) > 0;
-
-		if ( bRemoved )
-		{
-			delete Collection;
-			NumCollections--;
-		}
-	}
-
-	return bRemoved;
+	const FCollectionNameType CollectionKey(CollectionRef->GetCollectionName(), ShareType);
+	return CachedCollections.Remove(CollectionKey) > 0;
 }
 
 #undef LOCTEXT_NAMESPACE
diff --git a/Engine/Source/Developer/CollectionManager/Private/CollectionManager.h b/Engine/Source/Developer/CollectionManager/Private/CollectionManager.h
--- a/Engine/Source/Developer/CollectionManager/Private/CollectionManager.h
+++ b/Engine/Source/Developer/CollectionManager/Private/CollectionManager.h
@@ -14,17 +14,17 @@ public:
 	virtual bool GetAssetsInCollection(FName CollectionName, ECollectionShareType::Type ShareType, TArray<FName>& AssetPaths) const override;
 	virtual bool GetObjectsInCollection(FName CollectionName, ECollectionShareType::Type ShareType, TArray<FName>& ObjectPaths) const override;
 	virtual bool GetClassesInCollection(FName CollectionName, ECollectionShareType::Type ShareType, TArray<FName>& ClassPaths) const override;
-	virtual void GetCollectionsContainingAsset(FName ObjectPath, ECollectionShareType::Type ShareType, TArray<FName>& OutCollectionNames) const override;
+	virtual void GetCollectionsContainingObject(FName ObjectPath, ECollectionShareType::Type ShareType, TArray<FName>& OutCollectionNames) const override;
 	virtual void CreateUniqueCollectionName(const FName& BaseName, ECollectionShareType::Type ShareType, FName& OutCollectionName) const override;
 	virtual bool CreateCollection(FName CollectionName, ECollectionShareType::Type ShareType) override;
 	virtual bool RenameCollection(FName CurrentCollectionName, ECollectionShareType::Type CurrentShareType, FName NewCollectionName, ECollectionShareType::Type NewShareType) override;
 	virtual bool DestroyCollection(FName CollectionName, ECollectionShareType::Type ShareType) override;
 	virtual bool AddToCollection(FName CollectionName, ECollectionShareType::Type ShareType, FName ObjectPath) override;
-	virtual bool AddToCollection(FName CollectionName, ECollectionShareType::Type ShareType, const TArray<FName>& ObjectPaths, int32* OutNumAdded = NULL) override;
+	virtual bool AddToCollection(FName CollectionName, ECollectionShareType::Type ShareType, const TArray<FName>& ObjectPaths, int32* OutNumAdded = nullptr) override;
 	virtual bool RemoveFromCollection(FName CollectionName, ECollectionShareType::Type ShareType, FName ObjectPath) override;
-	virtual bool RemoveFromCollection(FName CollectionName, ECollectionShareType::Type ShareType, const TArray<FName>& ObjectPaths, int32* OutNumRemoved = NULL) override;
+	virtual bool RemoveFromCollection(FName CollectionName, ECollectionShareType::Type ShareType, const TArray<FName>& ObjectPaths, int32* OutNumRemoved = nullptr) override;
 	virtual bool EmptyCollection(FName CollectionName, ECollectionShareType::Type ShareType) override;
-	virtual bool IsCollectionEmpty(FName CollectionName, ECollectionShareType::Type ShareType) override;
+	virtual bool IsCollectionEmpty(FName CollectionName, ECollectionShareType::Type ShareType) const override;
 	virtual FText GetLastError() const override { return LastError; }
 
 	/** Event for when collections are created */
@@ -55,10 +55,10 @@ private:
 	bool ShouldUseSCC(ECollectionShareType::Type ShareType) const;
 
 	/** Adds a collection to the lookup maps */
-	class FCollection* AddCollection(const FCollection& Collection, ECollectionShareType::Type ShareType);
+	bool AddCollection(const TSharedRef<FCollection>& CollectionRef, ECollectionShareType::Type ShareType);
 
 	/** Removes a collection from the lookup maps */
-	bool RemoveCollection(FCollection* Collection, ECollectionShareType::Type ShareType);
+	bool RemoveCollection(const TSharedRef<FCollection>& CollectionRef, ECollectionShareType::Type ShareType);
 
 private:
 	/** The folders that contain collections */
@@ -68,13 +68,10 @@ private:
 	FString CollectionExtension;
 
 	/** A map of collection names to FCollection objects */
-	TMap<FName, FCollection*> CachedCollections[ECollectionShareType::CST_All];
-
-	/** A count of the number of collections created so we can be sure they are all destroyed when we shut down */
-	int32 NumCollections;
+	TMap<FCollectionNameType, TSharedRef<FCollection>> CachedCollections;
 
 	/** The most recent error that occurred */
-	FText LastError;
+	mutable FText LastError;
 
 	/** Event for when assets are added to a collection */
 	FAssetsAddedEvent AssetsAddedEvent;
@@ -90,4 +87,4 @@ private:
 
 	/** Event for when collections are destroyed */
 	FCollectionDestroyedEvent CollectionDestroyedEvent;
-};
\ No newline at end of file
+};
diff --git a/Engine/Source/Developer/CollectionManager/Public/CollectionManagerTypes.h b/Engine/Source/Developer/CollectionManager/Public/CollectionManagerTypes.h
--- a/Engine/Source/Developer/CollectionManager/Public/CollectionManagerTypes.h
+++ b/Engine/Source/Developer/CollectionManager/Public/CollectionManagerTypes.h
@@ -28,6 +28,14 @@ struct FCollectionNameType
 		return Name == Other.Name && Type == Other.Type;
 	}
 
+	friend inline uint32 GetTypeHash( const FCollectionNameType& Key )
+	{
+		uint32 Hash = 0;
+		HashCombine(Hash, GetTypeHash(Key.Name));
+		HashCombine(Hash, GetTypeHash(Key.Type));
+		return Hash;
+	}
+
 	FName Name;
 	ECollectionShareType::Type Type;
-};
\ No newline at end of file
+};
diff --git a/Engine/Source/Developer/CollectionManager/Public/ICollectionManager.h b/Engine/Source/Developer/CollectionManager/Public/ICollectionManager.h
--- a/Engine/Source/Developer/CollectionManager/Public/ICollectionManager.h
+++ b/Engine/Source/Developer/CollectionManager/Public/ICollectionManager.h
@@ -25,8 +25,8 @@ public:
 	/** Returns a list of object paths found in the specified collection and share type */
 	virtual bool GetObjectsInCollection(FName CollectionName, ECollectionShareType::Type ShareType, TArray<FName>& ObjectPaths) const = 0;
 
-	/** Returns a list of collections in which the specified asset exists of the specified share type */
-	virtual void GetCollectionsContainingAsset(FName ObjectPath, ECollectionShareType::Type ShareType, TArray<FName>& OutCollectionNames) const = 0;
+	/** Returns a list of collections in which the specified object exists of the specified share type */
+	virtual void GetCollectionsContainingObject(FName ObjectPath, ECollectionShareType::Type ShareType, TArray<FName>& OutCollectionNames) const = 0;
 
 	/** Creates a unique collection name for the given type taking the form BaseName+(unique number) */
 	virtual void CreateUniqueCollectionName(const FName& BaseName, ECollectionShareType::Type ShareType, FName& OutCollectionName) const = 0;
@@ -70,7 +70,7 @@ public:
 	  * @return true if the add was successful. If false, GetLastError will return a human readable string description of the error.
 	  */
 	virtual bool AddToCollection(FName CollectionName, ECollectionShareType::Type ShareType, FName ObjectPath) = 0;
-	virtual bool AddToCollection(FName CollectionName, ECollectionShareType::Type ShareType, const TArray<FName>& ObjectPaths, int32* OutNumAdded = NULL) = 0;
+	virtual bool AddToCollection(FName CollectionName, ECollectionShareType::Type ShareType, const TArray<FName>& ObjectPaths, int32* OutNumAdded = nullptr) = 0;
 
 	/**
 	  * Removes the asset from the specified collection.
@@ -82,7 +82,7 @@ public:
 	  * @return true if the remove was successful. If false, GetLastError will return a human readable string description of the error.
 	  */
 	virtual bool RemoveFromCollection(FName CollectionName, ECollectionShareType::Type ShareType, FName ObjectPath) = 0;
-	virtual bool RemoveFromCollection(FName CollectionName, ECollectionShareType::Type ShareType, const TArray<FName>& ObjectPaths, int32* OutNumRemoved = NULL) = 0;
+	virtual bool RemoveFromCollection(FName CollectionName, ECollectionShareType::Type ShareType, const TArray<FName>& ObjectPaths, int32* OutNumRemoved = nullptr) = 0;
 
 	/**
 	  * Removes all assets from the specified collection.
@@ -98,7 +98,7 @@ public:
 	  * @param ShareType The way the collection is shared.
 	  * @return true if the collection is empty.
 	  */
-	virtual bool IsCollectionEmpty(FName CollectionName, ECollectionShareType::Type ShareType) = 0;
+	virtual bool IsCollectionEmpty(FName CollectionName, ECollectionShareType::Type ShareType) const = 0;
 
 	/** Returns the most recent error. */
 	virtual FText GetLastError() const = 0;
@@ -122,4 +122,4 @@ public:
 	/** Event for when collections are renamed */
 	DECLARE_EVENT_TwoParams( ICollectionManager, FCollectionRenamedEvent, const FCollectionNameType&, const FCollectionNameType& );
 	virtual FCollectionRenamedEvent& OnCollectionRenamed() = 0;
-};
\ No newline at end of file
+};
