Optimized getting the list of collections for an object

The collections manager now stores a reverse look-up from the objects within its collections, to the collections that contain that object. This avoids the GetCollectionsContainingObject and GetCollectionsContainingObjects functions having to iterate the collections for each object queried, which can be a large win for large sets of objects.

This also adds some functions needed for the quick asset management feature.

diff --git a/Engine/Source/Developer/CollectionManager/Private/CollectionManager.cpp b/Engine/Source/Developer/CollectionManager/Private/CollectionManager.cpp
--- a/Engine/Source/Developer/CollectionManager/Private/CollectionManager.cpp
+++ b/Engine/Source/Developer/CollectionManager/Private/CollectionManager.cpp
@@ -21,6 +21,11 @@ FCollectionManager::~FCollectionManager()
 {
 }
 
+bool FCollectionManager::HasCollections() const
+{
+	return CachedCollections.Num() > 0;
+}
+
 void FCollectionManager::GetCollectionNames(ECollectionShareType::Type ShareType, TArray<FName>& CollectionNames) const
 {
 	for (const auto& CachedCollection : CachedCollections)
@@ -145,35 +150,64 @@ bool FCollectionManager::GetObjectsInCollection(FName CollectionName, ECollectio
 
 void FCollectionManager::GetCollectionsContainingObject(FName ObjectPath, ECollectionShareType::Type ShareType, TArray<FName>& OutCollectionNames) const
 {
-	for (const auto& CachedCollection : CachedCollections)
+	const auto* ObjectCollectionsPtr = CachedObjects.Find(ObjectPath);
+	if (ObjectCollectionsPtr)
 	{
-		const FCollectionNameType& CollectionKey = CachedCollection.Key;
-		const TSharedRef<FCollection>& Collection = CachedCollection.Value;
-		if ((ShareType == ECollectionShareType::CST_All || ShareType == CollectionKey.Type) && Collection->IsObjectInCollection(ObjectPath))
+		for (const FCollectionNameType& CollectionKey : *ObjectCollectionsPtr)
 		{
-			OutCollectionNames.Add(CollectionKey.Name);
+			if (ShareType == ECollectionShareType::CST_All || ShareType == CollectionKey.Type)
+			{
+				OutCollectionNames.Add(CollectionKey.Name);
+			}
 		}
 	}
 }
 
-FString FCollectionManager::GetCollectionsStringForObject(FName ObjectPath, ECollectionShareType::Type ShareType) const
+void FCollectionManager::GetCollectionsContainingObject(FName ObjectPath, TArray<FCollectionNameType>& OutCollections) const
 {
-	TArray<FString> CollectionNameStrings;
+	const auto* ObjectCollectionsPtr = CachedObjects.Find(ObjectPath);
+	if (ObjectCollectionsPtr)
+	{
+		OutCollections.Append(*ObjectCollectionsPtr);
+	}
+}
 
-	for (const auto& CachedCollection : CachedCollections)
+void FCollectionManager::GetCollectionsContainingObjects(const TArray<FName>& ObjectPaths, TMap<FCollectionNameType, TArray<FName>>& OutCollectionsAndMatchedObjects) const
+{
+	for (const FName& ObjectPath : ObjectPaths)
 	{
-		const FCollectionNameType& CollectionKey = CachedCollection.Key;
-		const TSharedRef<FCollection>& Collection = CachedCollection.Value;
-		if ((ShareType == ECollectionShareType::CST_All || ShareType == CollectionKey.Type) && Collection->IsObjectInCollection(ObjectPath))
+		const auto* ObjectCollectionsPtr = CachedObjects.Find(ObjectPath);
+		if (ObjectCollectionsPtr)
 		{
-			CollectionNameStrings.Add(CollectionKey.Name.ToString());
+			for (const FCollectionNameType& CollectionKey : *ObjectCollectionsPtr)
+			{
+				TArray<FName>& MatchedObjects = OutCollectionsAndMatchedObjects.FindOrAdd(CollectionKey);
+				MatchedObjects.Add(ObjectPath);
+			}
 		}
 	}
-	
-	if (CollectionNameStrings.Num() > 0)
+}
+
+FString FCollectionManager::GetCollectionsStringForObject(FName ObjectPath, ECollectionShareType::Type ShareType) const
+{
+	const auto* ObjectCollectionsPtr = CachedObjects.Find(ObjectPath);
+	if (ObjectCollectionsPtr)
 	{
-		CollectionNameStrings.Sort();
-		return FString::Join(CollectionNameStrings, TEXT(", "));
+		TArray<FString> CollectionNameStrings;
+
+		for (const FCollectionNameType& CollectionKey : *ObjectCollectionsPtr)
+		{
+			if (ShareType == ECollectionShareType::CST_All || ShareType == CollectionKey.Type)
+			{
+				CollectionNameStrings.Add(CollectionKey.Name.ToString());
+			}
+		}
+
+		if (CollectionNameStrings.Num() > 0)
+		{
+			CollectionNameStrings.Sort();
+			return FString::Join(CollectionNameStrings, TEXT(", "));
+		}
 	}
 
 	return FString();
@@ -283,6 +317,8 @@ bool FCollectionManager::RenameCollection(FName CurrentCollectionName, ECollecti
 		return false;
 	}
 
+	RebuildCachedObjects();
+
 	// Success
 	const FCollectionNameType OriginalCollectionKey(CurrentCollectionName, CurrentShareType);
 	const FCollectionNameType NewCollectionKey(NewCollectionName, NewShareType);
@@ -311,6 +347,9 @@ bool FCollectionManager::DestroyCollection(FName CollectionName, ECollectionShar
 	if ((*CollectionRefPtr)->DeleteSourceFile(LastError))
 	{
 		RemoveCollection(*CollectionRefPtr, ShareType);
+
+		RebuildCachedObjects();
+
 		CollectionDestroyedEvent.Broadcast(CollectionKey);
 		return true;
 	}
@@ -370,6 +409,8 @@ bool FCollectionManager::AddToCollection(FName CollectionName, ECollectionShareT
 				*OutNumAdded = NumAdded;
 			}
 
+			RebuildCachedObjects();
+
 			AssetsAddedEvent.Broadcast(CollectionKey, ObjectPaths);
 			return true;
 		}
@@ -445,6 +486,8 @@ bool FCollectionManager::RemoveFromCollection(FName CollectionName, ECollectionS
 			*OutNumRemoved = RemovedAssets.Num();
 		}
 
+		RebuildCachedObjects();
+
 		AssetsRemovedEvent.Broadcast(CollectionKey, ObjectPaths);
 		return true;
 	}
@@ -481,7 +524,13 @@ bool FCollectionManager::EmptyCollection(FName CollectionName, ECollectionShareT
 		return true;
 	}
 	
-	return RemoveFromCollection(CollectionName, ShareType, ObjectPaths);
+	if (RemoveFromCollection(CollectionName, ShareType, ObjectPaths))
+	{
+		RebuildCachedObjects();
+		return true;
+	}
+
+	return false;
 }
 
 bool FCollectionManager::IsCollectionEmpty(FName CollectionName, ECollectionShareType::Type ShareType) const
@@ -503,6 +552,25 @@ bool FCollectionManager::IsCollectionEmpty(FName CollectionName, ECollectionShar
 	return true;
 }
 
+bool FCollectionManager::IsObjectInCollection(FName ObjectPath, FName CollectionName, ECollectionShareType::Type ShareType) const
+{
+	if (!ensure(ShareType < ECollectionShareType::CST_All))
+	{
+		// Bad share type
+		LastError = LOCTEXT("Error_Internal", "There was an internal error.");
+		return true;
+	}
+
+	const FCollectionNameType CollectionKey(CollectionName, ShareType);
+	const TSharedRef<FCollection>* const CollectionRefPtr = CachedCollections.Find(CollectionKey);
+	if (CollectionRefPtr)
+	{
+		return (*CollectionRefPtr)->IsObjectInCollection(ObjectPath);
+	}
+
+	return false;
+}
+
 void FCollectionManager::LoadCollections()
 {
 	const double LoadStartTime = FPlatformTime::Seconds();
@@ -534,6 +602,32 @@ void FCollectionManager::LoadCollections()
 	}
 
 	UE_LOG(LogCollectionManager, Log, TEXT( "Loaded %d collections in %0.6f seconds" ), CachedCollections.Num() - PrevNumCollections, FPlatformTime::Seconds() - LoadStartTime);
+
+	RebuildCachedObjects();
+}
+
+void FCollectionManager::RebuildCachedObjects()
+{
+	const double LoadStartTime = FPlatformTime::Seconds();
+
+	CachedObjects.Empty();
+
+	for (const auto& CachedCollection : CachedCollections)
+	{
+		const FCollectionNameType& CollectionKey = CachedCollection.Key;
+		const TSharedRef<FCollection>& Collection = CachedCollection.Value;
+
+		TArray<FName> ObjectsInCollection;
+		Collection->GetObjectsInCollection(ObjectsInCollection);
+
+		for (const FName& ObjectPath : ObjectsInCollection)
+		{
+			auto& ObjectCollections = CachedObjects.FindOrAdd(ObjectPath);
+			ObjectCollections.AddUnique(CollectionKey);
+		}
+	}
+
+	UE_LOG(LogCollectionManager, Log, TEXT( "Rebuilt the cache for %d objects in %0.6f seconds" ), CachedObjects.Num(), FPlatformTime::Seconds() - LoadStartTime);
 }
 
 bool FCollectionManager::ShouldUseSCC(ECollectionShareType::Type ShareType) const
diff --git a/Engine/Source/Developer/CollectionManager/Private/CollectionManager.h b/Engine/Source/Developer/CollectionManager/Private/CollectionManager.h
--- a/Engine/Source/Developer/CollectionManager/Private/CollectionManager.h
+++ b/Engine/Source/Developer/CollectionManager/Private/CollectionManager.h
@@ -9,12 +9,15 @@ public:
 	virtual ~FCollectionManager();
 
 	// ICollectionManager implementation
+	virtual bool HasCollections() const override;
 	virtual void GetCollectionNames(ECollectionShareType::Type ShareType, TArray<FName>& CollectionNames) const override;
 	virtual bool CollectionExists(FName CollectionName, ECollectionShareType::Type ShareType) const override;
 	virtual bool GetAssetsInCollection(FName CollectionName, ECollectionShareType::Type ShareType, TArray<FName>& AssetPaths) const override;
 	virtual bool GetObjectsInCollection(FName CollectionName, ECollectionShareType::Type ShareType, TArray<FName>& ObjectPaths) const override;
 	virtual bool GetClassesInCollection(FName CollectionName, ECollectionShareType::Type ShareType, TArray<FName>& ClassPaths) const override;
 	virtual void GetCollectionsContainingObject(FName ObjectPath, ECollectionShareType::Type ShareType, TArray<FName>& OutCollectionNames) const override;
+	virtual void GetCollectionsContainingObject(FName ObjectPath, TArray<FCollectionNameType>& OutCollections) const override;
+	virtual void GetCollectionsContainingObjects(const TArray<FName>& ObjectPaths, TMap<FCollectionNameType, TArray<FName>>& OutCollectionsAndMatchedObjects) const override;
 	virtual FString GetCollectionsStringForObject(FName ObjectPath, ECollectionShareType::Type ShareType) const override;
 	virtual void CreateUniqueCollectionName(const FName& BaseName, ECollectionShareType::Type ShareType, FName& OutCollectionName) const override;
 	virtual bool CreateCollection(FName CollectionName, ECollectionShareType::Type ShareType) override;
@@ -26,6 +29,7 @@ public:
 	virtual bool RemoveFromCollection(FName CollectionName, ECollectionShareType::Type ShareType, const TArray<FName>& ObjectPaths, int32* OutNumRemoved = nullptr) override;
 	virtual bool EmptyCollection(FName CollectionName, ECollectionShareType::Type ShareType) override;
 	virtual bool IsCollectionEmpty(FName CollectionName, ECollectionShareType::Type ShareType) const override;
+	virtual bool IsObjectInCollection(FName ObjectPath, FName CollectionName, ECollectionShareType::Type ShareType) const override;
 	virtual FText GetLastError() const override { return LastError; }
 
 	/** Event for when collections are created */
@@ -52,6 +56,9 @@ private:
 	/** Loads all collection files from disk */
 	void LoadCollections();
 
+	/** Rebuild the entire cached objects map based on the current collection data */
+	void RebuildCachedObjects();
+
 	/** Returns true if the specified share type requires source control */
 	bool ShouldUseSCC(ECollectionShareType::Type ShareType) const;
 
@@ -71,6 +78,9 @@ private:
 	/** A map of collection names to FCollection objects */
 	TMap<FCollectionNameType, TSharedRef<FCollection>> CachedCollections;
 
+	/** A map of object paths to their associated collection names */
+	TMap<FName, TArray<FCollectionNameType>> CachedObjects;
+
 	/** The most recent error that occurred */
 	mutable FText LastError;
 
diff --git a/Engine/Source/Developer/CollectionManager/Public/CollectionManagerTypes.h b/Engine/Source/Developer/CollectionManager/Public/CollectionManagerTypes.h
--- a/Engine/Source/Developer/CollectionManager/Public/CollectionManagerTypes.h
+++ b/Engine/Source/Developer/CollectionManager/Public/CollectionManagerTypes.h
@@ -63,6 +63,38 @@ struct ECollectionShareType
 		}
 		return FText::GetEmpty();
 	}
+
+	static FText GetDescription(const Type InType)
+	{
+		switch (InType)
+		{
+		case CST_Local:
+			return NSLOCTEXT("ECollectionShareType", "CST_Local_Description", "Local. This collection is only visible to you and is not in source control.");
+		case CST_Private:
+			return NSLOCTEXT("ECollectionShareType", "CST_Private_Description", "Private. This collection is only visible to you.");
+		case CST_Shared:
+			return NSLOCTEXT("ECollectionShareType", "CST_Shared_Description", "Shared. This collection is visible to everyone.");
+		default:
+			break;
+		}
+		return FText::GetEmpty();
+	}
+
+	static FName GetIconStyleName(const Type InType)
+	{
+		switch (InType)
+		{
+		case CST_Local:
+			return FName("ContentBrowser.Local");
+		case CST_Private:
+			return FName("ContentBrowser.Private");
+		case CST_Shared:
+			return FName("ContentBrowser.Shared");
+		default:
+			break;
+		}
+		return NAME_None;
+	}
 };
 
 /** A name/type pair to uniquely identify a collection */
diff --git a/Engine/Source/Developer/CollectionManager/Public/ICollectionManager.h b/Engine/Source/Developer/CollectionManager/Public/ICollectionManager.h
--- a/Engine/Source/Developer/CollectionManager/Public/ICollectionManager.h
+++ b/Engine/Source/Developer/CollectionManager/Public/ICollectionManager.h
@@ -10,6 +10,9 @@ public:
 	/** Virtual destructor */
 	virtual ~ICollectionManager() {}
 
+	/** Returns whether or not the collection manager contains any collections */
+	virtual bool HasCollections() const = 0;
+
 	/** Returns the list of collection names of the specified share type */
 	virtual void GetCollectionNames(ECollectionShareType::Type ShareType, TArray<FName>& CollectionNames) const = 0;
 
@@ -28,6 +31,12 @@ public:
 	/** Returns a list of collections in which the specified object exists of the specified share type */
 	virtual void GetCollectionsContainingObject(FName ObjectPath, ECollectionShareType::Type ShareType, TArray<FName>& OutCollectionNames) const = 0;
 
+	/** Returns a list of collections in which the specified object exists */
+	virtual void GetCollectionsContainingObject(FName ObjectPath, TArray<FCollectionNameType>& OutCollections) const = 0;
+
+	/** Returns a list of collections in which any of the specified objects exist */
+	virtual void GetCollectionsContainingObjects(const TArray<FName>& ObjectPaths, TMap<FCollectionNameType, TArray<FName>>& OutCollectionsAndMatchedObjects) const = 0;
+
 	/** Returns a string containing a comma separated list of collections in which the specified object exists of the specified share type */
 	virtual FString GetCollectionsStringForObject(FName ObjectPath, ECollectionShareType::Type ShareType) const = 0;
 
@@ -103,6 +112,16 @@ public:
 	  */
 	virtual bool IsCollectionEmpty(FName CollectionName, ECollectionShareType::Type ShareType) const = 0;
 
+	/**
+	 * Check to see if the given object exists in the given collection
+	 *
+	 * @param ObjectPath The object to test for its existence in the collection.
+	 * @param CollectionName The collection to test.
+	 * @param ShareType The way the collection is shared.
+	 * @return true if the object is in the collection.
+	 */
+	virtual bool IsObjectInCollection(FName ObjectPath, FName CollectionName, ECollectionShareType::Type ShareType) const = 0;
+
 	/** Returns the most recent error. */
 	virtual FText GetLastError() const = 0;
 
