Made collections robust against moving or renaming assets and classes

UETOOL-332 - Collections 2.0
UETOOL-373 - Check robustness of Collections 2.0

The collection manager will now fix-up any references to redirectors within its collections once the asset registry has finished discovering all the available assets. It also adds a watcher so it can pick up any moves or deletes as they happen.

All of this fix-up is applied to the in-memory copy of the collection, and doesn't get persisted to disk until it needs to be (because referenced redirectors are being deleted). This helps to minimize issues with source control availability and shared/private collections.

This change also makes sure that objects that are being referenced by a collection will leave a redirector behind, and also adds some extra context (such as the collection name) to some of the source control errors that may be reported when saving a collection.

diff --git a/Engine/Source/Developer/AssetTools/AssetTools.Build.cs b/Engine/Source/Developer/AssetTools/AssetTools.Build.cs
--- a/Engine/Source/Developer/AssetTools/AssetTools.Build.cs
+++ b/Engine/Source/Developer/AssetTools/AssetTools.Build.cs
@@ -33,6 +33,7 @@ public class AssetTools : ModuleRules
 				"Analytics",
 				"AssetRegistry",
 				"ContentBrowser",
+				"CollectionManager",
                 "CurveAssetEditor",
 				"DesktopPlatform",
 				"EditorWidgets",
@@ -53,6 +54,7 @@ public class AssetTools : ModuleRules
 			new string[] {
 				"AssetRegistry",
 				"ContentBrowser",
+				"CollectionManager",
                 "CurveAssetEditor",
 				"CurveTableEditor",
 				"DataTableEditor",
diff --git a/Engine/Source/Developer/AssetTools/Private/AssetFixUpRedirectors.cpp b/Engine/Source/Developer/AssetTools/Private/AssetFixUpRedirectors.cpp
--- a/Engine/Source/Developer/AssetTools/Private/AssetFixUpRedirectors.cpp
+++ b/Engine/Source/Developer/AssetTools/Private/AssetFixUpRedirectors.cpp
@@ -3,6 +3,7 @@
 
 #include "AssetToolsPrivatePCH.h"
 #include "AssetRegistryModule.h"
+#include "CollectionManagerModule.h"
 #include "ISourceControlModule.h"
 #include "ObjectTools.h"
 #include "MessageLog.h"
@@ -13,12 +14,14 @@
 struct FRedirectorRefs
 {
 	UObjectRedirector* Redirector;
+	FName RedirectorPackageName;
 	TArray<FName> ReferencingPackageNames;
 	FText FailureReason;
 	bool bRedirectorValidForFixup;
 
 	FRedirectorRefs(UObjectRedirector* InRedirector)
 		: Redirector(InRedirector)
+		, RedirectorPackageName(InRedirector->GetOutermost()->GetFName())
 		, bRedirectorValidForFixup(true)
 	{}
 };
@@ -91,6 +94,9 @@ void FAssetFixUpRedirectors::ExecuteFixUp(TArray<TWeakObjectPtr<UObjectRedirecto
 				TArray<UPackage*> FailedToSave;
 				SaveReferencingPackages(ReferencingPackagesToSave, FailedToSave);
 
+				// Save any collections that were referencing any of the redirectors
+				SaveReferencingCollections(RedirectorRefsList);
+
 				// Wait for package referencers to be updated
 				UpdateAssetReferencers(RedirectorRefsList);
 
@@ -110,7 +116,7 @@ void FAssetFixUpRedirectors::PopulateRedirectorReferencers(TArray<FRedirectorRef
 	for ( auto RedirectorRefsIt = RedirectorsToPopulate.CreateIterator(); RedirectorRefsIt; ++RedirectorRefsIt )
 	{
 		FRedirectorRefs& RedirectorRefs = *RedirectorRefsIt;
-		AssetRegistryModule.Get().GetReferencers(RedirectorRefs.Redirector->GetOutermost()->GetFName(), RedirectorRefs.ReferencingPackageNames);
+		AssetRegistryModule.Get().GetReferencers(RedirectorRefs.RedirectorPackageName, RedirectorRefs.ReferencingPackageNames);
 	}
 }
 
@@ -296,6 +302,29 @@ void FAssetFixUpRedirectors::SaveReferencingPackages(const TArray<UPackage*>& Re
 	}
 }
 
+void FAssetFixUpRedirectors::SaveReferencingCollections(TArray<FRedirectorRefs>& RedirectorsToFix) const
+{
+	FAssetRegistryModule& AssetRegistryModule = FModuleManager::LoadModuleChecked<FAssetRegistryModule>(TEXT("AssetRegistry"));
+	FCollectionManagerModule& CollectionManagerModule = FCollectionManagerModule::GetModule();
+
+	// Find all collections that were referenced by any of the redirectors that are potentially going to be removed and attempt to re-save them
+	// The redirectors themselves will have already been fixed up, as collections do that once the asset registry has been populated, 
+	// however collections lazily re-save redirector fix-up to avoid SCC issues, so we need to force that now
+	for (FRedirectorRefs& RedirectorRefs : RedirectorsToFix)
+	{
+		// Follow each link in the redirector, and notify the collections manager that it is going to be removed - this will force it to re-save any required collections
+		for (UObjectRedirector* Redirector = RedirectorRefs.Redirector; Redirector; Redirector = Cast<UObjectRedirector>(Redirector->DestinationObject))
+		{
+			const FName RedirectorObjectPath = *Redirector->GetPathName();
+			if (!CollectionManagerModule.Get().HandleRedirectorDeleted(RedirectorObjectPath))
+			{
+				RedirectorRefs.FailureReason = FText::Format(LOCTEXT("RedirectorFixupFailed_CollectionsFailedToSave", "Referencing collection(s) failed to save: {0}"), CollectionManagerModule.Get().GetLastError());
+				RedirectorRefs.bRedirectorValidForFixup = false;
+			}
+		}
+	}
+}
+
 void FAssetFixUpRedirectors::UpdateAssetReferencers(const TArray<FRedirectorRefs>& RedirectorsToFix) const
 {
 	// Load the asset registry module
@@ -304,7 +333,7 @@ void FAssetFixUpRedirectors::UpdateAssetReferencers(const TArray<FRedirectorRefs
 	TArray<FString> AssetPaths;
 	for (const auto& Redirector : RedirectorsToFix)
 	{
-		AssetPaths.AddUnique(FPackageName::GetLongPackagePath(Redirector.Redirector->GetOutermost()->GetName()));
+		AssetPaths.AddUnique(FPackageName::GetLongPackagePath(Redirector.RedirectorPackageName.ToString()) / TEXT("")); // Ensure trailing slash
 	}
 	AssetRegistryModule.Get().ScanPathsSynchronous(AssetPaths, true);
 }
@@ -395,7 +424,7 @@ void FAssetFixUpRedirectors::ReportFailures(const TArray<FRedirectorRefs>& Redir
 				bTitleOutput = true;
 			}
 			FFormatNamedArguments Arguments;
-			Arguments.Add(TEXT("PackageName"), FText::FromString(RedirectorRefs.Redirector->GetOutermost()->GetName()));
+			Arguments.Add(TEXT("PackageName"), FText::FromName(RedirectorRefs.RedirectorPackageName));
 			Arguments.Add(TEXT("FailureReason"), FText::FromString(RedirectorRefs.FailureReason.ToString()));
 			EditorErrors.Warning(FText::Format(LOCTEXT("RedirectorFixupFailedReason", "{PackageName} - {FailureReason}"), Arguments ));
 		}
diff --git a/Engine/Source/Developer/AssetTools/Private/AssetFixUpRedirectors.h b/Engine/Source/Developer/AssetTools/Private/AssetFixUpRedirectors.h
--- a/Engine/Source/Developer/AssetTools/Private/AssetFixUpRedirectors.h
+++ b/Engine/Source/Developer/AssetTools/Private/AssetFixUpRedirectors.h
@@ -43,6 +43,9 @@ private:
 	/** Saves all the referencing packages and updates SCC state */
 	void SaveReferencingPackages(const TArray<UPackage*>& ReferencingPackagesToSave, TArray<UPackage*>& OutFailedToSave) const;
 
+	/** Saves any collections that were referencing any of the redirectors and updates SCC state */
+	void SaveReferencingCollections(TArray<FRedirectorRefs>& RedirectorsToFix) const;
+
 	/** Waits for the asset registry to update its asset referencer cache */
 	void UpdateAssetReferencers(const TArray<FRedirectorRefs>& RedirectorsToFix) const;
 
diff --git a/Engine/Source/Developer/AssetTools/Private/AssetRenameManager.cpp b/Engine/Source/Developer/AssetTools/Private/AssetRenameManager.cpp
--- a/Engine/Source/Developer/AssetTools/Private/AssetRenameManager.cpp
+++ b/Engine/Source/Developer/AssetTools/Private/AssetRenameManager.cpp
@@ -4,6 +4,7 @@
 #include "AssetToolsPrivatePCH.h"
 #include "AssetRegistryModule.h"
 #include "AssetToolsModule.h"
+#include "CollectionManagerModule.h"
 #include "ISourceControlModule.h"
 #include "FileHelpers.h"
 #include "ObjectTools.h"
@@ -165,9 +166,9 @@ void FAssetRenameManager::FixReferencesAndRename(TArray<FAssetRenameData> Assets
 	// Prep a list of assets to rename with an extra boolean to determine if they should leave a redirector or not
 	TArray<FAssetRenameDataWithReferencers> AssetsToRename;
 	AssetsToRename.Reset(AssetsAndNames.Num());
-	for (int32 AssetIdx = 0; AssetIdx < AssetsAndNames.Num(); ++AssetIdx)
+	for (const FAssetRenameData& AssetRenameData : AssetsAndNames)
 	{
-		new(AssetsToRename)FAssetRenameDataWithReferencers(AssetsAndNames[AssetIdx]);
+		AssetsToRename.Emplace(FAssetRenameDataWithReferencers(AssetRenameData));
 	}
 
 	// Warn the user if they are about to rename an asset that is referenced by a CDO
@@ -199,6 +200,9 @@ void FAssetRenameManager::FixReferencesAndRename(TArray<FAssetRenameData> Assets
 	// Update the source control state for the packages containing the assets we are renaming if source control is enabled. If source control is enabled and this fails we can not continue.
 	if ( UpdatePackageStatus(AssetsToRename) )
 	{
+		// Detect whether the assets are being referenced by a collection. Assets within a collection must leave a redirector to avoid the collection losing its references.
+		DetectReferencingCollections(AssetsToRename);
+
 		// Load all referencing packages and mark any assets that must have redirectors.
 		TArray<UPackage*> ReferencingPackagesToSave;
 		LoadReferencingPackages(AssetsToRename, ReferencingPackagesToSave);
@@ -465,6 +469,25 @@ bool FAssetRenameManager::CheckOutPackages(TArray<FAssetRenameDataWithReferencer
 	return bUserAcceptedCheckout;
 }
 
+void FAssetRenameManager::DetectReferencingCollections(TArray<FAssetRenameDataWithReferencers>& AssetsToRename) const
+{
+	FCollectionManagerModule& CollectionManagerModule = FCollectionManagerModule::GetModule();
+
+	for (auto& AssetToRename : AssetsToRename)
+	{
+		if (AssetToRename.Asset.IsValid())
+		{
+			TArray<FCollectionNameType> ReferencingCollections;
+			CollectionManagerModule.Get().GetCollectionsContainingObject(*AssetToRename.Asset->GetPathName(), ReferencingCollections);
+
+			if (ReferencingCollections.Num() > 0)
+			{
+				AssetToRename.bCreateRedirector = true;
+			}
+		}
+	}
+}
+
 void FAssetRenameManager::DetectReadOnlyPackages(TArray<FAssetRenameDataWithReferencers>& AssetsToRename, TArray<UPackage*>& InOutReferencingPackagesToSave) const
 {
 	// For each valid package...
diff --git a/Engine/Source/Developer/AssetTools/Private/AssetRenameManager.h b/Engine/Source/Developer/AssetTools/Private/AssetRenameManager.h
--- a/Engine/Source/Developer/AssetTools/Private/AssetRenameManager.h
+++ b/Engine/Source/Developer/AssetTools/Private/AssetRenameManager.h
@@ -57,6 +57,9 @@ private:
 	  */
 	bool CheckOutPackages(TArray<FAssetRenameDataWithReferencers>& AssetsToRename, TArray<UPackage*>& InOutReferencingPackagesToSave) const;
 
+	/** Finds any collections that are referencing the assets to be renamed. Assets referenced by collections will leave redirectors */
+	void DetectReferencingCollections(TArray<FAssetRenameDataWithReferencers>& AssetsToRename) const;
+
 	/** Finds any read only packages and removes them from the save list. Assets referenced by these packages will leave redirectors. */ 
 	void DetectReadOnlyPackages(TArray<FAssetRenameDataWithReferencers>& AssetsToRename, TArray<UPackage*>& InOutReferencingPackagesToSave) const;
 
diff --git a/Engine/Source/Developer/CollectionManager/Private/Collection.cpp b/Engine/Source/Developer/CollectionManager/Private/Collection.cpp
--- a/Engine/Source/Developer/CollectionManager/Private/Collection.cpp
+++ b/Engine/Source/Developer/CollectionManager/Private/Collection.cpp
@@ -117,7 +117,7 @@ bool FCollection::Load(FText& OutError)
 
 			if ( Line.Len() )
 			{
-				AddAssetToCollection(FName(*Line));
+				AddObjectToCollection(FName(*Line));
 			}
 		}
 	}
@@ -200,8 +200,8 @@ bool FCollection::Save(FText& OutError)
 		}
 		else
 		{
-			OutError = LOCTEXT("Error_WriteFailed", "Failed to write to collection file.");
-			UE_LOG(LogCollectionManager, Error, TEXT("%s %s"), *OutError.ToString(), *CollectionName.ToString());
+			OutError = FText::Format(LOCTEXT("Error_WriteFailed", "Failed to write to collection file: {0}"), FText::FromString(SourceFilename));
+			UE_LOG(LogCollectionManager, Error, TEXT("%s"), *OutError.ToString());
 		}
 	}
 	else
@@ -272,7 +272,7 @@ bool FCollection::DeleteSourceFile(FText& OutError)
 			bSuccessfullyDeleted = IFileManager::Get().Delete(*SourceFilename);
 			if ( !bSuccessfullyDeleted )
 			{
-				OutError = LOCTEXT("Error_DiskDeleteFailed", "Failed to delete the collection file from disk.");
+				OutError = FText::Format(LOCTEXT("Error_DiskDeleteFailed", "Failed to delete the collection file: {0}"), FText::FromString(SourceFilename));
 			}
 		}
 	}
@@ -290,9 +290,8 @@ bool FCollection::DeleteSourceFile(FText& OutError)
 	return bSuccessfullyDeleted;
 }
 
-bool FCollection::AddAssetToCollection(FName ObjectPath)
+bool FCollection::AddObjectToCollection(FName ObjectPath)
 {
-	// @todo collection Does this apply to dynamic collections?
 	if (!ObjectSet.Contains(ObjectPath))
 	{
 		ObjectSet.Add(ObjectPath);
@@ -302,15 +301,13 @@ bool FCollection::AddAssetToCollection(FName ObjectPath)
 	return false;
 }
 
-bool FCollection::RemoveAssetFromCollection(FName ObjectPath)
+bool FCollection::RemoveObjectFromCollection(FName ObjectPath)
 {
-	// @todo collection Does this apply to dynamic collections?
 	return ObjectSet.Remove(ObjectPath) > 0;
 }
 
 void FCollection::GetAssetsInCollection(TArray<FName>& Assets) const
 {
-	// @todo collection Does this apply to dynamic collections?
 	for (const FName& ObjectName : ObjectSet)
 	{
 		if (!ObjectName.ToString().StartsWith(TEXT("/Script/")))
@@ -322,7 +319,6 @@ void FCollection::GetAssetsInCollection(TArray<FName>& Assets) const
 
 void FCollection::GetClassesInCollection(TArray<FName>& Classes) const
 {
-	// @todo collection Does this apply to dynamic collections?
 	for (const FName& ObjectName : ObjectSet)
 	{
 		if (ObjectName.ToString().StartsWith(TEXT("/Script/")))
@@ -334,16 +330,21 @@ void FCollection::GetClassesInCollection(TArray<FName>& Classes) const
 
 void FCollection::GetObjectsInCollection(TArray<FName>& Objects) const
 {
-	// @todo collection Does this apply to dynamic collections?
 	FCollectionUtils::AppendCollectionToArray(ObjectSet, Objects);
 }
 
 bool FCollection::IsObjectInCollection(FName ObjectPath) const
 {
-	// @todo collection Does this apply to dynamic collections?
 	return ObjectSet.Contains(ObjectPath);
 }
 
+bool FCollection::IsRedirectorInCollection(FName ObjectPath) const
+{
+	// Redirectors are fixed up in-memory once the asset registry has finished loading, 
+	// so we need to test our on-disk set of objects rather than our in-memory set of objects
+	return DiskSnapshot.ObjectSet.Contains(ObjectPath);
+}
+
 bool FCollection::IsDynamic() const
 {
 	// @todo collection Dynamic collections
@@ -526,7 +527,7 @@ bool FCollection::CheckoutCollection(FText& OutError)
 		if ( SourceControlProvider.Execute(ISourceControlOperation::Create<FSync>(), AbsoluteFilename) == ECommandResult::Failed )
 		{
 			// Could not sync up with the head revision
-			OutError = LOCTEXT("Error_SCCSync", "Failed to sync collection to the head revision.");
+			OutError = FText::Format(LOCTEXT("Error_SCCSync", "Failed to sync collection '{0}' to the head revision."), FText::FromName(CollectionName));
 			return false;
 		}
 
@@ -539,7 +540,7 @@ bool FCollection::CheckoutCollection(FText& OutError)
 			if ( !NewCollection.Load(LoadErrorText) )
 			{
 				// Failed to load the head revision file so it isn't safe to delete it
-				OutError = FText::Format(LOCTEXT("Error_SCCBadHead", "Failed to load the collection at the head revision. {0}"), LoadErrorText);
+				OutError = FText::Format(LOCTEXT("Error_SCCBadHead", "Failed to load the collection '{0}' at the head revision. {1}"), FText::FromName(CollectionName), LoadErrorText);
 				return false;
 			}
 
@@ -569,20 +570,20 @@ bool FCollection::CheckoutCollection(FText& OutError)
 			bSuccessfullyCheckedOut = (SourceControlProvider.Execute(ISourceControlOperation::Create<FCheckOut>(), AbsoluteFilename) == ECommandResult::Succeeded);
 			if (!bSuccessfullyCheckedOut)
 			{
-				OutError = LOCTEXT("Error_SCCCheckout", "Failed to check out collection.");
+				OutError = FText::Format(LOCTEXT("Error_SCCCheckout", "Failed to check out collection '{0}'"), FText::FromName(CollectionName));
 			}
 		}
 		else if(!SourceControlState->IsCurrent())
 		{
-			OutError = LOCTEXT("Error_SCCNotCurrent", "Collection is not at head revision after sync.");
+			OutError = FText::Format(LOCTEXT("Error_SCCNotCurrent", "Collection '{0}' is not at head revision after sync."), FText::FromName(CollectionName));
 		}
 		else if(SourceControlState->IsCheckedOutOther())
 		{
-			OutError = LOCTEXT("Error_SCCCheckedOutOther", "Collection is checked out by another user.");
+			OutError = FText::Format(LOCTEXT("Error_SCCCheckedOutOther", "Collection '{0}' is checked out by another user."), FText::FromName(CollectionName));
 		}
 		else
 		{
-			OutError = LOCTEXT("Error_SCCUnknown", "Could not determine source control state.");
+			OutError = FText::Format(LOCTEXT("Error_SCCUnknown", "Could not determine source control state for collection '{0}'"), FText::FromName(CollectionName));
 		}
 	}
 	else
@@ -623,7 +624,7 @@ bool FCollection::CheckinCollection(FText& OutError)
 		const bool bWasAdded = (SourceControlProvider.Execute(ISourceControlOperation::Create<FMarkForAdd>(), AbsoluteFilename) == ECommandResult::Succeeded);
 		if (!bWasAdded)
 		{
-			OutError = LOCTEXT("Error_SCCAdd", "Failed to add new collection to source control.");
+			OutError = FText::Format(LOCTEXT("Error_SCCAdd", "Failed to add collection '{0}' to source control."), FText::FromName(CollectionName));
 			return false;
 		}
 		SourceControlState = SourceControlProvider.GetState(AbsoluteFilename, EStateCacheUsage::ForceUpdate);
@@ -631,7 +632,7 @@ bool FCollection::CheckinCollection(FText& OutError)
 
 	if ( SourceControlState.IsValid() && !(SourceControlState->IsCheckedOut() || SourceControlState->IsAdded()) )
 	{
-		OutError = LOCTEXT("Error_SCCNotCheckedOut", "Collection not checked out or open for add.");
+		OutError = FText::Format(LOCTEXT("Error_SCCNotCheckedOut", "Collection '{0}' not checked out or open for add."), FText::FromName(CollectionName));
 		return false;
 	}
 
@@ -729,7 +730,7 @@ bool FCollection::CheckinCollection(FText& OutError)
 	}
 	else 
 	{
-		OutError = LOCTEXT("Error_SCCCheckIn", "Failed to check in collection.");
+		OutError = FText::Format(LOCTEXT("Error_SCCCheckIn", "Failed to check in collection '{0}'."), FText::FromName(CollectionName));
 		return false;
 	}
 }
@@ -760,7 +761,7 @@ bool FCollection::RevertCollection(FText& OutError)
 
 	if ( SourceControlState.IsValid() && !(SourceControlState->IsCheckedOut() || SourceControlState->IsAdded()) )
 	{
-		OutError = LOCTEXT("Error_SCCNotCheckedOut", "Collection not checked out or open for add.");
+		OutError = FText::Format(LOCTEXT("Error_SCCNotCheckedOut", "Collection '{0}' not checked out or open for add."), FText::FromName(CollectionName));
 		return false;
 	}
 
@@ -770,7 +771,7 @@ bool FCollection::RevertCollection(FText& OutError)
 	}
 	else
 	{
-		OutError = LOCTEXT("Error_SCCRevert", "Could not revert Collection.");
+		OutError = FText::Format(LOCTEXT("Error_SCCRevert", "Could not revert collection '{0}'"), FText::FromName(CollectionName));
 		return false;
 	}
 }
@@ -830,7 +831,7 @@ bool FCollection::DeleteFromSourceControl(FText& OutError)
 		{
 			// Could not sync up with the head revision
 			GWarn->EndSlowTask();
-			OutError = LOCTEXT("Error_SCCSync", "Failed to sync collection to the head revision.");
+			OutError = FText::Format(LOCTEXT("Error_SCCSync", "Failed to sync collection '{0}' to the head revision."), FText::FromName(CollectionName));
 			return false;
 		}
 
@@ -848,7 +849,7 @@ bool FCollection::DeleteFromSourceControl(FText& OutError)
 		{
 			// Failed to load the head revision file so it isn't safe to delete it
 			GWarn->EndSlowTask();
-			OutError = FText::Format(LOCTEXT("Error_SCCBadHead", "Failed to load the collection at the head revision. {0}"), LoadErrorText);
+			OutError = FText::Format(LOCTEXT("Error_SCCBadHead", "Failed to load the collection '{0}' at the head revision. {1}"), FText::FromName(CollectionName), LoadErrorText);
 			return false;
 		}
 
@@ -865,7 +866,7 @@ bool FCollection::DeleteFromSourceControl(FText& OutError)
 	{
 		if(SourceControlState->IsAdded() || SourceControlState->IsCheckedOut())
 		{
-			OutError = LOCTEXT("Error_SCCDeleteWhileCheckedOut", "Failed to delete collection in source control because it is checked out or open for add.");
+			OutError = FText::Format(LOCTEXT("Error_SCCDeleteWhileCheckedOut", "Failed to delete collection '{0}' in source control because it is checked out or open for add."), FText::FromName(CollectionName));
 		}
 		else if(SourceControlState->CanCheckout())
 		{
@@ -888,12 +889,12 @@ bool FCollection::DeleteFromSourceControl(FText& OutError)
 						UE_LOG(LogCollectionManager, Warning, TEXT("Failed to revert collection '%s' after failing to check in the file that was marked for delete."), *CollectionName.ToString());
 					}
 
-					OutError = LOCTEXT("Error_SCCCheckIn", "Failed to check in collection.");
+					OutError = FText::Format(LOCTEXT("Error_SCCCheckIn", "Failed to check in collection '{0}'."), FText::FromName(CollectionName));
 				}
 			}
 			else
 			{
-				OutError = LOCTEXT("Error_SCCDeleteFailed", "Failed to delete collection in source control.");
+				OutError = FText::Format(LOCTEXT("Error_SCCDeleteFailed", "Failed to delete collection '{0}' in source control."), FText::FromName(CollectionName));
 			}
 		}
 		else if(!SourceControlState->IsSourceControlled())
@@ -902,20 +903,20 @@ bool FCollection::DeleteFromSourceControl(FText& OutError)
 			bDeletedSuccessfully = IFileManager::Get().Delete(*AbsoluteFilename);
 			if ( !bDeletedSuccessfully )
 			{
-				OutError = LOCTEXT("Error_DiskDeleteFailed", "Failed to delete the collection file from disk.");
+				OutError = FText::Format(LOCTEXT("Error_DiskDeleteFailed", "Failed to delete the collection file: {0}"), FText::FromString(AbsoluteFilename));
 			}
 		}
 		else if (!SourceControlState->IsCurrent())
 		{
-			OutError = LOCTEXT("Error_SCCNotCurrent", "Collection is not at head revision after sync.");
+			OutError = FText::Format(LOCTEXT("Error_SCCNotCurrent", "Collection '{0}' is not at head revision after sync."), FText::FromName(CollectionName));
 		}
 		else if(SourceControlState->IsCheckedOutOther())
 		{
-			OutError = LOCTEXT("Error_SCCCheckedOutOther", "Collection is checked out by another user.");
+			OutError = FText::Format(LOCTEXT("Error_SCCCheckedOutOther", "Collection '{0}' is checked out by another user."), FText::FromName(CollectionName));
 		}
 		else
 		{
-			OutError = LOCTEXT("Error_SCCUnknown", "Could not determine source control state.");
+			OutError = FText::Format(LOCTEXT("Error_SCCUnknown", "Could not determine source control state for collection '{0}'"), FText::FromName(CollectionName));
 		}
 	}
 	else
diff --git a/Engine/Source/Developer/CollectionManager/Private/Collection.h b/Engine/Source/Developer/CollectionManager/Private/Collection.h
--- a/Engine/Source/Developer/CollectionManager/Private/Collection.h
+++ b/Engine/Source/Developer/CollectionManager/Private/Collection.h
@@ -47,10 +47,10 @@ public:
 	/** Deletes the source file for this collection. If false, OutError is a human readable warning depicting the error. */
 	bool DeleteSourceFile(FText& OutError);
 
-	/** Adds a single asset to the collection. Static collections only. */
-	bool AddAssetToCollection(FName ObjectPath);
-	/** Removes a single asset from the collection. Static collections only. */
-	bool RemoveAssetFromCollection(FName ObjectPath);
+	/** Adds a single object to the collection. Static collections only. */
+	bool AddObjectToCollection(FName ObjectPath);
+	/** Removes a single object from the collection. Static collections only. */
+	bool RemoveObjectFromCollection(FName ObjectPath);
 	/** Gets a list of assets in the collection. Static collections only. */
 	void GetAssetsInCollection(TArray<FName>& Assets) const;
 	/** Gets a list of classes in the collection. Static collections only. */
@@ -59,6 +59,8 @@ public:
 	void GetObjectsInCollection(TArray<FName>& Objects) const;
 	/** Returns true when the specified object is in the collection. Static collections only. */
 	bool IsObjectInCollection(FName ObjectPath) const;
+	/** Returns true when the specified redirector is in the collection. Static collections only. */
+	bool IsRedirectorInCollection(FName ObjectPath) const;
 
 	/** Returns true if the collection contains rules instead of a flat list */
 	bool IsDynamic() const;
diff --git a/Engine/Source/Developer/CollectionManager/Private/CollectionManager.cpp b/Engine/Source/Developer/CollectionManager/Private/CollectionManager.cpp
--- a/Engine/Source/Developer/CollectionManager/Private/CollectionManager.cpp
+++ b/Engine/Source/Developer/CollectionManager/Private/CollectionManager.cpp
@@ -638,7 +638,7 @@ bool FCollectionManager::AddToCollection(FName CollectionName, ECollectionShareT
 	int32 NumAdded = 0;
 	for (const FName& ObjectPath : ObjectPaths)
 	{
-		if ((*CollectionRefPtr)->AddAssetToCollection(ObjectPath))
+		if ((*CollectionRefPtr)->AddObjectToCollection(ObjectPath))
 		{
 			NumAdded++;
 		}
@@ -664,7 +664,7 @@ bool FCollectionManager::AddToCollection(FName CollectionName, ECollectionShareT
 			// Added but not saved, revert the add
 			for (const FName& ObjectPath : ObjectPaths)
 			{
-				(*CollectionRefPtr)->RemoveAssetFromCollection(ObjectPath);
+				(*CollectionRefPtr)->RemoveObjectFromCollection(ObjectPath);
 			}
 			return false;
 		}
@@ -710,7 +710,7 @@ bool FCollectionManager::RemoveFromCollection(FName CollectionName, ECollectionS
 	TArray<FName> RemovedAssets;
 	for (const FName& ObjectPath : ObjectPaths)
 	{
-		if ((*CollectionRefPtr)->RemoveAssetFromCollection(ObjectPath))
+		if ((*CollectionRefPtr)->RemoveObjectFromCollection(ObjectPath))
 		{
 			RemovedAssets.Add(ObjectPath);
 		}
@@ -741,7 +741,7 @@ bool FCollectionManager::RemoveFromCollection(FName CollectionName, ECollectionS
 		// Removed but not saved, revert the remove
 		for (const FName& RemovedAssetName : RemovedAssets)
 		{
-			(*CollectionRefPtr)->AddAssetToCollection(RemovedAssetName);
+			(*CollectionRefPtr)->AddObjectToCollection(RemovedAssetName);
 		}
 		return false;
 	}
@@ -859,6 +859,127 @@ bool FCollectionManager::IsValidParentCollection(FName CollectionName, ECollecti
 	return bValidParent;
 }
 
+void FCollectionManager::HandleFixupRedirectors(ICollectionRedirectorFollower& InRedirectorFollower)
+{
+	const double LoadStartTime = FPlatformTime::Seconds();
+
+	TArray<TPair<FName, FName>> ObjectsToRename;
+
+	// Build up the list of redirected object into rename pairs
+	for (const auto& CachedObjectInfo : CachedObjects)
+	{
+		FName NewObjectPath;
+		if (InRedirectorFollower.FixupObject(CachedObjectInfo.Key, NewObjectPath))
+		{
+			ObjectsToRename.Add(TPairInitializer<FName, FName>(CachedObjectInfo.Key, NewObjectPath));
+		}
+	}
+
+	// Notify the rename for each redirected object
+	for (const auto& ObjectToRename : ObjectsToRename)
+	{
+		HandleObjectRenamed(ObjectToRename.Key, ObjectToRename.Value);
+	}
+
+	UE_LOG(LogCollectionManager, Log, TEXT( "Fixed up redirectors for %d collections in %0.6f seconds (updated %d objects)" ), CachedCollections.Num(), FPlatformTime::Seconds() - LoadStartTime, ObjectsToRename.Num());
+
+	for (const auto& ObjectToRename : ObjectsToRename)
+	{
+		UE_LOG(LogCollectionManager, Verbose, TEXT( "\tRedirected '%s' to '%s'" ), *ObjectToRename.Key.ToString(), *ObjectToRename.Value.ToString());
+	}
+}
+
+bool FCollectionManager::HandleRedirectorDeleted(const FName& ObjectPath)
+{
+	bool bSavedAllCollections = true;
+
+	FTextBuilder AllErrors;
+
+	// We don't have a cache for on-disk objects, so we have to do this the slower way and query each collection in turn
+	for (const auto& CachedCollection : CachedCollections)
+	{
+		const TSharedRef<FCollection>& Collection = CachedCollection.Value;
+
+		if (Collection->IsRedirectorInCollection(ObjectPath))
+		{
+			FText SaveError;
+			if (!Collection->Save(SaveError))
+			{
+				AllErrors.AppendLine(SaveError);
+				bSavedAllCollections = false;
+			}
+		}
+	}
+
+	if (!bSavedAllCollections)
+	{
+		LastError = AllErrors.ToText();
+	}
+
+	return bSavedAllCollections;
+}
+
+void FCollectionManager::HandleObjectRenamed(const FName& OldObjectPath, const FName& NewObjectPath)
+{
+	// Remove the info about the collections that the old object path is within - we'll add this info back in under the new object path shortly
+	TArray<FObjectCollectionInfo> OldObjectCollectionInfos;
+	if (!CachedObjects.RemoveAndCopyValue(OldObjectPath, OldObjectCollectionInfos))
+	{
+		return;
+	}
+
+	// Replace this object reference in all collections that use it, and update our object cache
+	auto& NewObjectCollectionInfos = CachedObjects.FindOrAdd(NewObjectPath);
+	for (const FObjectCollectionInfo& OldObjectCollectionInfo : OldObjectCollectionInfos)
+	{
+		if ((OldObjectCollectionInfo.Reason & ECollectionRecursionFlags::Self) != 0)
+		{
+			// The old object is contained directly within this collection (rather than coming from a parent or child collection), so update the object reference
+			const TSharedRef<FCollection>* const CollectionRefPtr = CachedCollections.Find(OldObjectCollectionInfo.CollectionKey);
+			if (CollectionRefPtr)
+			{
+				(*CollectionRefPtr)->RemoveObjectFromCollection(OldObjectPath);
+				(*CollectionRefPtr)->AddObjectToCollection(NewObjectPath);
+			}
+		}
+
+		// Merge the collection references for the old object with any collection references that already exist for the new object
+		FObjectCollectionInfo* ObjectInfoPtr = NewObjectCollectionInfos.FindByPredicate([&](const FObjectCollectionInfo& InCollectionInfo) { return InCollectionInfo.CollectionKey == OldObjectCollectionInfo.CollectionKey; });
+		if (ObjectInfoPtr)
+		{
+			ObjectInfoPtr->Reason |= OldObjectCollectionInfo.Reason;
+		}
+		else
+		{
+			NewObjectCollectionInfos.Add(OldObjectCollectionInfo);
+		}
+	}
+}
+
+void FCollectionManager::HandleObjectDeleted(const FName& ObjectPath)
+{
+	// Remove the info about the collections that the old object path is within
+	TArray<FObjectCollectionInfo> ObjectCollectionInfos;
+	if (!CachedObjects.RemoveAndCopyValue(ObjectPath, ObjectCollectionInfos))
+	{
+		return;
+	}
+
+	// Remove this object reference from all collections that use it
+	for (const FObjectCollectionInfo& ObjectCollectionInfo : ObjectCollectionInfos)
+	{
+		if ((ObjectCollectionInfo.Reason & ECollectionRecursionFlags::Self) != 0)
+		{
+			// The object is contained directly within this collection (rather than coming from a parent or child collection), so remove the object reference
+			const TSharedRef<FCollection>* const CollectionRefPtr = CachedCollections.Find(ObjectCollectionInfo.CollectionKey);
+			if (CollectionRefPtr)
+			{
+				(*CollectionRefPtr)->RemoveObjectFromCollection(ObjectPath);
+			}
+		}
+	}
+}
+
 void FCollectionManager::LoadCollections()
 {
 	const double LoadStartTime = FPlatformTime::Seconds();
@@ -908,7 +1029,7 @@ void FCollectionManager::RebuildCachedObjects()
 		const TSharedRef<FCollection>& Collection = CachedCollection.Value;
 
 		TArray<FName> ObjectsInCollection;
-		CachedCollection.Value->GetObjectsInCollection(ObjectsInCollection);
+		Collection->GetObjectsInCollection(ObjectsInCollection);
 
 		if (ObjectsInCollection.Num() > 0)
 		{
diff --git a/Engine/Source/Developer/CollectionManager/Private/CollectionManager.h b/Engine/Source/Developer/CollectionManager/Private/CollectionManager.h
--- a/Engine/Source/Developer/CollectionManager/Private/CollectionManager.h
+++ b/Engine/Source/Developer/CollectionManager/Private/CollectionManager.h
@@ -38,6 +38,10 @@ public:
 	virtual bool IsObjectInCollection(FName ObjectPath, FName CollectionName, ECollectionShareType::Type ShareType, ECollectionRecursionFlags::Flags RecursionMode = ECollectionRecursionFlags::Self) const override;
 	virtual bool IsValidParentCollection(FName CollectionName, ECollectionShareType::Type ShareType, FName ParentCollectionName, ECollectionShareType::Type ParentShareType) const override;
 	virtual FText GetLastError() const override { return LastError; }
+	virtual void HandleFixupRedirectors(ICollectionRedirectorFollower& InRedirectorFollower) override;
+	virtual bool HandleRedirectorDeleted(const FName& ObjectPath) override;
+	virtual void HandleObjectRenamed(const FName& OldObjectPath, const FName& NewObjectPath) override;
+	virtual void HandleObjectDeleted(const FName& ObjectPath) override;
 
 	/** Event for when collections are created */
 	DECLARE_DERIVED_EVENT( FCollectionManager, ICollectionManager::FCollectionCreatedEvent, FCollectionCreatedEvent );
diff --git a/Engine/Source/Developer/CollectionManager/Public/CollectionManagerTypes.h b/Engine/Source/Developer/CollectionManager/Public/CollectionManagerTypes.h
--- a/Engine/Source/Developer/CollectionManager/Public/CollectionManagerTypes.h
+++ b/Engine/Source/Developer/CollectionManager/Public/CollectionManagerTypes.h
@@ -143,3 +143,12 @@ struct FCollectionNameType
 	FName Name;
 	ECollectionShareType::Type Type;
 };
+
+class ICollectionRedirectorFollower
+{
+public:
+	virtual ~ICollectionRedirectorFollower() {}
+
+	/** Given an object path, will see if it needs to follow any redirectors, and if so, will populate OutNewObjectPath with the new name and return true */
+	virtual bool FixupObject(const FName& InObjectPath, FName& OutNewObjectPath) = 0;
+};
diff --git a/Engine/Source/Developer/CollectionManager/Public/ICollectionManager.h b/Engine/Source/Developer/CollectionManager/Public/ICollectionManager.h
--- a/Engine/Source/Developer/CollectionManager/Public/ICollectionManager.h
+++ b/Engine/Source/Developer/CollectionManager/Public/ICollectionManager.h
@@ -162,6 +162,24 @@ public:
 	/** Returns the most recent error. */
 	virtual FText GetLastError() const = 0;
 
+	/**
+	 * Called to notify the collections that they should fix-up their object references so that they no longer contain any redirectors
+	 * References are only updated in-memory, and won't be saved to disk until a redirector is deleted (which forces our hand), or the collection is saved for any other reason
+	 */
+	virtual void HandleFixupRedirectors(ICollectionRedirectorFollower& InRedirectorFollower) = 0;
+
+	/**
+	 * Called to notify the collections that a redirector has been deleted and that they should ensure their on-disk representation is re-saved with the fixed up in-memory version
+	 * @return true if all of the collections that were referencing this redirector could be re-saved, false otherwise
+	 */
+	virtual bool HandleRedirectorDeleted(const FName& ObjectPath) = 0;
+
+	/** Called to notify the collections that an object has been renamed or moved */
+	virtual void HandleObjectRenamed(const FName& OldObjectPath, const FName& NewObjectPath) = 0;
+
+	/** Called to notify the collections that an object has been deleted */
+	virtual void HandleObjectDeleted(const FName& ObjectPath) = 0;
+
 	/** Event for when collections are created */
 	DECLARE_EVENT_OneParam( ICollectionManager, FCollectionCreatedEvent, const FCollectionNameType& );
 	virtual FCollectionCreatedEvent& OnCollectionCreated() = 0;
diff --git a/Engine/Source/Editor/ContentBrowser/Private/CollectionAssetRegistryBridge.cpp b/Engine/Source/Editor/ContentBrowser/Private/CollectionAssetRegistryBridge.cpp
new file mode 100644
--- /dev/null
+++ b/Engine/Source/Editor/ContentBrowser/Private/CollectionAssetRegistryBridge.cpp
@@ -0,0 +1,141 @@
+// Copyright 1998-2015 Epic Games, Inc. All Rights Reserved.
+
+#include "ContentBrowserPCH.h"
+#include "CollectionAssetRegistryBridge.h"
+
+#define LOCTEXT_NAMESPACE "ContentBrowser"
+
+/** The collection manager doesn't know how to follow redirectors, this class provides it with that knowledge */
+class FCollectionRedirectorFollower : public ICollectionRedirectorFollower
+{
+public:
+	FCollectionRedirectorFollower()
+		: AssetRegistryModule(FModuleManager::LoadModuleChecked<FAssetRegistryModule>(TEXT("AssetRegistry")))
+	{
+	}
+
+	virtual bool FixupObject(const FName& InObjectPath, FName& OutNewObjectPath) override
+	{
+		OutNewObjectPath = NAME_None;
+
+		if (InObjectPath.ToString().StartsWith(TEXT("/Script/")))
+		{
+			// We can't use FindObject while we're saving
+			if (!GIsSavingPackage)
+			{
+				const FString ClassPathStr = InObjectPath.ToString();
+
+				UClass* FoundClass = FindObject<UClass>(ANY_PACKAGE, *ClassPathStr);
+				if (!FoundClass)
+				{
+					// Use the linker to search for class name redirects (from the loaded ActiveClassRedirects)
+					const FString ClassName = FPackageName::ObjectPathToObjectName(ClassPathStr);
+					const FName NewClassName = FLinkerLoad::FindNewNameForClass(*ClassName, false);
+
+					if (!NewClassName.IsNone())
+					{
+						// Our new class name might be lacking the path, so try and find it so we can use the full path in the collection
+						FoundClass = FindObject<UClass>(ANY_PACKAGE, *NewClassName.ToString());
+						if (FoundClass)
+						{
+							OutNewObjectPath = *FoundClass->GetPathName();
+						}
+					}
+				}
+			}
+		}
+		else
+		{
+			// Use the asset registry to avoid loading the object
+			FAssetData ObjectAssetData = AssetRegistryModule.Get().GetAssetByObjectPath(InObjectPath);
+			while (ObjectAssetData.IsValid() && ObjectAssetData.IsRedirector())
+			{
+				// Get the destination object from the meta-data rather than load the redirector object, as 
+				// loading a redirector will also load the object it points to, which could cause a large hitch
+				const FString* DestinationObjectStrPtr = ObjectAssetData.TagsAndValues.Find("DestinationObject");
+				if (DestinationObjectStrPtr)
+				{
+					FString DestinationObjectPath = *DestinationObjectStrPtr;
+					ConstructorHelpers::StripObjectClass(DestinationObjectPath);
+					ObjectAssetData = AssetRegistryModule.Get().GetAssetByObjectPath(*DestinationObjectPath);
+				}
+				else
+				{
+					ObjectAssetData = FAssetData();
+				}
+			}
+
+			OutNewObjectPath = ObjectAssetData.ObjectPath;
+		}
+
+		return OutNewObjectPath != NAME_None && InObjectPath != OutNewObjectPath;
+	}
+
+private:
+	FAssetRegistryModule& AssetRegistryModule;
+};
+
+FCollectionAssetRegistryBridge::FCollectionAssetRegistryBridge()
+{
+	// Load the asset registry module to listen for updates
+	FAssetRegistryModule& AssetRegistryModule = FModuleManager::LoadModuleChecked<FAssetRegistryModule>(TEXT("AssetRegistry"));
+	AssetRegistryModule.Get().OnAssetRemoved().AddRaw(this, &FCollectionAssetRegistryBridge::OnAssetRemoved);
+	AssetRegistryModule.Get().OnAssetRenamed().AddRaw(this, &FCollectionAssetRegistryBridge::OnAssetRenamed);
+	
+	if (AssetRegistryModule.Get().IsLoadingAssets())
+	{
+		AssetRegistryModule.Get().OnFilesLoaded().AddRaw(this, &FCollectionAssetRegistryBridge::OnAssetRegistryLoadComplete);
+	}
+	else
+	{
+		OnAssetRegistryLoadComplete();
+	}
+}
+
+FCollectionAssetRegistryBridge::~FCollectionAssetRegistryBridge()
+{
+	// Load the asset registry module to unregister delegates
+	if (FModuleManager::Get().IsModuleLoaded("AssetRegistry"))
+	{
+		FAssetRegistryModule& AssetRegistryModule = FModuleManager::GetModuleChecked<FAssetRegistryModule>(TEXT("AssetRegistry"));
+		AssetRegistryModule.Get().OnAssetRemoved().RemoveAll(this);
+		AssetRegistryModule.Get().OnAssetRenamed().RemoveAll(this);
+		AssetRegistryModule.Get().OnFilesLoaded().RemoveAll(this);
+	}
+}
+
+void FCollectionAssetRegistryBridge::OnAssetRegistryLoadComplete()
+{
+	FCollectionManagerModule& CollectionManagerModule = FCollectionManagerModule::GetModule();
+
+	// We've found all the assets, let the collections manager fix up its references now so that it doesn't reference any redirectors
+	FCollectionRedirectorFollower RedirectorFollower;
+	CollectionManagerModule.Get().HandleFixupRedirectors(RedirectorFollower);
+}
+
+void FCollectionAssetRegistryBridge::OnAssetRenamed(const FAssetData& AssetData, const FString& OldObjectPath)
+{
+	FCollectionManagerModule& CollectionManagerModule = FCollectionManagerModule::GetModule();
+
+	// Notify the collections manager that an asset has been renamed
+	CollectionManagerModule.Get().HandleObjectRenamed(*OldObjectPath, AssetData.ObjectPath);
+}
+
+void FCollectionAssetRegistryBridge::OnAssetRemoved(const FAssetData& AssetData)
+{
+	FCollectionManagerModule& CollectionManagerModule = FCollectionManagerModule::GetModule();
+
+	if (AssetData.IsRedirector())
+	{
+		// Notify the collections manager that a redirector has been removed
+		// This will attempt to re-save any collections that still have a reference to this redirector in their on-disk collection data
+		CollectionManagerModule.Get().HandleRedirectorDeleted(AssetData.ObjectPath);
+	}
+	else
+	{
+		// Notify the collections manager that an asset has been removed
+		CollectionManagerModule.Get().HandleObjectDeleted(AssetData.ObjectPath);
+	}
+}
+
+#undef LOCTEXT_NAMESPACE
diff --git a/Engine/Source/Editor/ContentBrowser/Private/CollectionAssetRegistryBridge.h b/Engine/Source/Editor/ContentBrowser/Private/CollectionAssetRegistryBridge.h
new file mode 100644
--- /dev/null
+++ b/Engine/Source/Editor/ContentBrowser/Private/CollectionAssetRegistryBridge.h
@@ -0,0 +1,23 @@
+// Copyright 1998-2015 Epic Games, Inc. All Rights Reserved.
+
+#pragma once
+
+/**
+ * The bridge between the asset registry and the collections manager - used to update collections as certain asset events happen
+ */
+class FCollectionAssetRegistryBridge
+{
+public:
+	FCollectionAssetRegistryBridge();
+	~FCollectionAssetRegistryBridge();
+
+private:
+	/** Called when the asset registry initial load has completed */
+	void OnAssetRegistryLoadComplete();
+
+	/** Handler for when an asset was renamed in the asset registry */
+	void OnAssetRenamed(const FAssetData& AssetData, const FString& OldObjectPath);
+
+	/** Handler for when an asset was removed from the asset registry */
+	void OnAssetRemoved(const FAssetData& AssetData);
+};
diff --git a/Engine/Source/Editor/ContentBrowser/Private/ContentBrowserSingleton.cpp b/Engine/Source/Editor/ContentBrowser/Private/ContentBrowserSingleton.cpp
--- a/Engine/Source/Editor/ContentBrowser/Private/ContentBrowserSingleton.cpp
+++ b/Engine/Source/Editor/ContentBrowser/Private/ContentBrowserSingleton.cpp
@@ -11,11 +11,13 @@
 #include "TutorialMetaData.h"
 #include "SDockTab.h"
 #include "NativeClassHierarchy.h"
+#include "CollectionAssetRegistryBridge.h"
 
 #define LOCTEXT_NAMESPACE "ContentBrowser"
 
 FContentBrowserSingleton::FContentBrowserSingleton()
-	: SettingsStringID(0)
+	: CollectionAssetRegistryBridge(MakeShareable(new FCollectionAssetRegistryBridge()))
+	, SettingsStringID(0)
 {
 	// Register the tab spawners for all content browsers
 	const FSlateIcon ContentBrowserIcon(FEditorStyle::GetStyleSetName(), "ContentBrowser.TabIcon");
diff --git a/Engine/Source/Editor/ContentBrowser/Private/ContentBrowserSingleton.h b/Engine/Source/Editor/ContentBrowser/Private/ContentBrowserSingleton.h
--- a/Engine/Source/Editor/ContentBrowser/Private/ContentBrowserSingleton.h
+++ b/Engine/Source/Editor/ContentBrowser/Private/ContentBrowserSingleton.h
@@ -5,6 +5,7 @@
 
 class SContentBrowser;
 class FNativeClassHierarchy;
+class FCollectionAssetRegistryBridge;
 
 #define MAX_CONTENT_BROWSERS 4
 
@@ -33,7 +34,7 @@ public:
 	virtual void SyncBrowserToAssets(const TArray<UObject*>& AssetList, bool bAllowLockedBrowsers = false) override;
 	virtual void GetSelectedAssets(TArray<FAssetData>& SelectedAssets) override;
 
-		/** Gets the content browser singleton as a FContentBrowserSingleton */
+	/** Gets the content browser singleton as a FContentBrowserSingleton */
 	static FContentBrowserSingleton& Get();
 	
 	/** Sets the current primary content browser. */
@@ -88,6 +89,8 @@ private:
 
 	TSharedPtr<FNativeClassHierarchy> NativeClassHierarchy;
 
+	TSharedRef<FCollectionAssetRegistryBridge> CollectionAssetRegistryBridge;
+
 	/** An incrementing int32 which is used when making unique settings strings */
 	int32 SettingsStringID;
 };
\ No newline at end of file
