Added the parent info to the collections list shown in the asset tooltip

UE-17674 - Collection path in not displayed in asset tooltips

diff --git a/Engine/Source/Developer/CollectionManager/Private/CollectionManager.cpp b/Engine/Source/Developer/CollectionManager/Private/CollectionManager.cpp
--- a/Engine/Source/Developer/CollectionManager/Private/CollectionManager.cpp
+++ b/Engine/Source/Developer/CollectionManager/Private/CollectionManager.cpp
@@ -327,11 +327,22 @@ FString FCollectionManager::GetCollectionsStringForObject(FName ObjectPath, ECol
 	{
 		TArray<FString> CollectionNameStrings;
 
+		TArray<FString> CollectionPathStrings;
+
+		auto GetCollectionsStringForObjectWorker = [&](const FCollectionNameType& InCollectionKey, ECollectionRecursionFlags::Flag InReason) -> ERecursiveWorkerFlowControl
+		{
+			CollectionPathStrings.Insert(InCollectionKey.Name.ToString(), 0);
+			return ERecursiveWorkerFlowControl::Continue;
+		};
+
 		for (const FObjectCollectionInfo& ObjectCollectionInfo : *ObjectCollectionInfosPtr)
 		{
 			if ((ShareType == ECollectionShareType::CST_All || ShareType == ObjectCollectionInfo.CollectionKey.Type) && (RecursionMode & ObjectCollectionInfo.Reason) != 0)
 			{
-				CollectionNameStrings.Add(ObjectCollectionInfo.CollectionKey.Name.ToString());
+				CollectionPathStrings.Reset();
+				RecursionHelper_DoWork(ObjectCollectionInfo.CollectionKey, ECollectionRecursionFlags::SelfAndParents, GetCollectionsStringForObjectWorker);
+
+				CollectionNameStrings.Add(FString::Join(CollectionPathStrings, TEXT("/")));
 			}
 		}
 
