Limited the types of collections can be nested within one another

UE-17678 - Limit what type of collections can be nested

Local collections can now only contain other local collections.
Private collections can now only contain local or private collections.
Shared collections can still contain any other kind of collection.

diff --git a/Engine/Source/Developer/CollectionManager/Private/CollectionManager.cpp b/Engine/Source/Developer/CollectionManager/Private/CollectionManager.cpp
--- a/Engine/Source/Developer/CollectionManager/Private/CollectionManager.cpp
+++ b/Engine/Source/Developer/CollectionManager/Private/CollectionManager.cpp
@@ -877,6 +877,15 @@ bool FCollectionManager::IsValidParentCollection(FName CollectionName, ECollecti
 				: LOCTEXT("InvalidParent_CannotParentToChildren", "A collection cannot be parented to its children");
 			return ERecursiveWorkerFlowControl::Stop;
 		}
+
+		const bool bIsValidChildType = ECollectionShareType::IsValidChildType(InCollectionKey.Type, ShareType);
+		if (!bIsValidChildType)
+		{
+			bValidParent = false;
+			LastError = FText::Format(LOCTEXT("InvalidParent_InvalidChildType", "A {0} collection cannot contain a {1} collection"), ECollectionShareType::ToText(InCollectionKey.Type), ECollectionShareType::ToText(ShareType));
+			return ERecursiveWorkerFlowControl::Stop;
+		}
+
 		return ERecursiveWorkerFlowControl::Continue;
 	};
 
diff --git a/Engine/Source/Developer/CollectionManager/Public/CollectionManagerTypes.h b/Engine/Source/Developer/CollectionManager/Public/CollectionManagerTypes.h
--- a/Engine/Source/Developer/CollectionManager/Public/CollectionManagerTypes.h
+++ b/Engine/Source/Developer/CollectionManager/Public/CollectionManagerTypes.h
@@ -95,6 +95,29 @@ struct ECollectionShareType
 		}
 		return NAME_None;
 	}
+
+	static bool IsValidChildType(const Type InParentType, const Type InChildType)
+	{
+		if (InParentType == CST_Local)
+		{
+			// Local collections can only contain other local collections
+			return InChildType == CST_Local;
+		}
+
+		if (InParentType == CST_Private)
+		{
+			// Private collections can only contain local or private collections
+			return InChildType == CST_Local || InChildType == CST_Private;
+		}
+
+		if (InParentType == CST_Shared)
+		{
+			// Shared collections can contain any kind of collection
+			return InChildType == CST_Local || InChildType == CST_Private || InChildType == CST_Shared;
+		}
+
+		return false;
+	}
 };
 
 /** Controls how the collections manager will recurse when performing work against a given collection */
