Added stricter validation for collection names

UE-17561 - Collections allow for several special characters that are not normally allowed
UE-17549 - User gets no warning for invalid Characters when naming a Collection
UE-17538 - Add a Character limit for naming a collection

The list of valid characters now matches that of assets and folders in the Content Browser, and the user is informed when they try and use an invalid character.

The collection name is also now restricted to NAME_SIZE to avoid a crash if a user entered a name that was too long.

diff --git a/Engine/Source/Developer/CollectionManager/Private/CollectionManager.cpp b/Engine/Source/Developer/CollectionManager/Private/CollectionManager.cpp
--- a/Engine/Source/Developer/CollectionManager/Private/CollectionManager.cpp
+++ b/Engine/Source/Developer/CollectionManager/Private/CollectionManager.cpp
@@ -397,6 +397,40 @@ void FCollectionManager::CreateUniqueCollectionName(const FName& BaseName, EColl
 	while (CollectionAlreadyExists);
 }
 
+bool FCollectionManager::IsValidCollectionName(const FString& CollectionName, ECollectionShareType::Type ShareType) const
+{
+	// Make sure we are not creating an FName that is too large
+	if (CollectionName.Len() > NAME_SIZE)
+	{
+		LastError = LOCTEXT("Error_CollectionNameTooLong", "This collection name is too long. Please choose a shorter name.");
+		return false;
+	}
+
+	const FName CollectionNameFinal = *CollectionName;
+
+	// Make sure the we actually have a new name set
+	if (CollectionNameFinal.IsNone())
+	{
+		LastError = LOCTEXT("Error_CollectionNameEmptyOrNone", "This collection name cannot be empty or 'None'.");
+		return false;
+	}
+
+	// Make sure the new name only contains valid characters
+	if (!CollectionNameFinal.IsValidXName(INVALID_OBJECTNAME_CHARACTERS INVALID_LONGPACKAGE_CHARACTERS, &LastError))
+	{
+		return false;
+	}
+
+	// Make sure we're not duplicating an existing collection name
+	if (CollectionExists(CollectionNameFinal, ShareType))
+	{
+		LastError = FText::Format(LOCTEXT("Error_CollectionAlreadyExists", "A collection already exists with the name '{0}'."), FText::FromName(CollectionNameFinal));
+		return false;
+	}
+
+	return true;
+}
+
 bool FCollectionManager::CreateCollection(FName CollectionName, ECollectionShareType::Type ShareType)
 {
 	if (!ensure(ShareType < ECollectionShareType::CST_All))
diff --git a/Engine/Source/Developer/CollectionManager/Private/CollectionManager.h b/Engine/Source/Developer/CollectionManager/Private/CollectionManager.h
--- a/Engine/Source/Developer/CollectionManager/Private/CollectionManager.h
+++ b/Engine/Source/Developer/CollectionManager/Private/CollectionManager.h
@@ -26,6 +26,7 @@ public:
 	virtual void GetCollectionsContainingObjects(const TArray<FName>& ObjectPaths, TMap<FCollectionNameType, TArray<FName>>& OutCollectionsAndMatchedObjects, ECollectionRecursionFlags::Flags RecursionMode = ECollectionRecursionFlags::Self) const override;
 	virtual FString GetCollectionsStringForObject(FName ObjectPath, ECollectionShareType::Type ShareType, ECollectionRecursionFlags::Flags RecursionMode = ECollectionRecursionFlags::Self) const override;
 	virtual void CreateUniqueCollectionName(const FName& BaseName, ECollectionShareType::Type ShareType, FName& OutCollectionName) const override;
+	virtual bool IsValidCollectionName(const FString& CollectionName, ECollectionShareType::Type ShareType) const override;
 	virtual bool CreateCollection(FName CollectionName, ECollectionShareType::Type ShareType) override;
 	virtual bool RenameCollection(FName CurrentCollectionName, ECollectionShareType::Type CurrentShareType, FName NewCollectionName, ECollectionShareType::Type NewShareType) override;
 	virtual bool ReparentCollection(FName CollectionName, ECollectionShareType::Type ShareType, FName ParentCollectionName, ECollectionShareType::Type ParentShareType) override;
diff --git a/Engine/Source/Developer/CollectionManager/Public/ICollectionManager.h b/Engine/Source/Developer/CollectionManager/Public/ICollectionManager.h
--- a/Engine/Source/Developer/CollectionManager/Public/ICollectionManager.h
+++ b/Engine/Source/Developer/CollectionManager/Public/ICollectionManager.h
@@ -61,6 +61,9 @@ public:
 	/** Creates a unique collection name for the given type taking the form BaseName+(unique number) */
 	virtual void CreateUniqueCollectionName(const FName& BaseName, ECollectionShareType::Type ShareType, FName& OutCollectionName) const = 0;
 
+	/** Returns whether or not the given collection name is valid. If false, GetLastError will return a human readable string description of the error. */
+	virtual bool IsValidCollectionName(const FString& CollectionName, ECollectionShareType::Type ShareType) const = 0;
+
 	/**
 	  * Adds a collection to the asset registry. A .collection file will be added to disk.
 	  *
diff --git a/Engine/Source/Editor/ContentBrowser/Private/SCollectionView.cpp b/Engine/Source/Editor/ContentBrowser/Private/SCollectionView.cpp
--- a/Engine/Source/Editor/ContentBrowser/Private/SCollectionView.cpp
+++ b/Engine/Source/Editor/ContentBrowser/Private/SCollectionView.cpp
@@ -1332,22 +1332,17 @@ bool SCollectionView::CollectionNameChangeCommit( const TSharedPtr< FCollectionI
 
 bool SCollectionView::CollectionVerifyRenameCommit(const TSharedPtr< FCollectionItem >& CollectionItem, const FString& NewName, const FSlateRect& MessageAnchor, FText& OutErrorMessage)
 {
-	const FName NewNameFinal = *NewName;
-
 	// If the new name is the same as the old name, consider this to be unchanged, and accept it.
-	if (CollectionItem->CollectionName == NewNameFinal)
+	if (CollectionItem->CollectionName.ToString() == NewName)
 	{
 		return true;
 	}
 
 	FCollectionManagerModule& CollectionManagerModule = FCollectionManagerModule::GetModule();
 
-	if (CollectionManagerModule.Get().CollectionExists(NewNameFinal, ECollectionShareType::CST_All))
+	if (!CollectionManagerModule.Get().IsValidCollectionName(NewName, ECollectionShareType::CST_Shared))
 	{
-		// This collection already exists, inform the user and continue
-		OutErrorMessage = FText::Format(LOCTEXT("RenameCollectionAlreadyExists", "A collection already exists with the name '{0}'."), FText::FromName(NewNameFinal));
-
-		// Return false to indicate that the user should enter a new name
+		OutErrorMessage = CollectionManagerModule.Get().GetLastError();
 		return false;
 	}
 
diff --git a/Engine/Source/Editor/ContentBrowser/Private/SourcesViewWidgets.cpp b/Engine/Source/Editor/ContentBrowser/Private/SourcesViewWidgets.cpp
--- a/Engine/Source/Editor/ContentBrowser/Private/SourcesViewWidgets.cpp
+++ b/Engine/Source/Editor/ContentBrowser/Private/SourcesViewWidgets.cpp
@@ -355,8 +355,6 @@ void SCollectionTreeItem::Construct( const FArguments& InArgs )
 			]
 
 			+SHorizontalBox::Slot()
-			.AutoWidth()
-			.VAlign(VAlign_Center)
 			[
 				SAssignNew(InlineRenameWidget, SInlineEditableTextBlock)
 				.Text( this, &SCollectionTreeItem::GetNameText )
