Added status information to collections

This shows a little colored pip to the right of the collection items, where each color means:
  Grey   - Up-to-date and empty
  Green  - Up-to-date and not empty
  Blue   - Has local changes
  Orange - Not at the latest revision
  Red    - Checked out by another user, in a conflicted state, or missing its source control provider

This also adds an "Update" and "Save" option to the collections context menu. Typically these won't be needed as collections mostly manage themselves, but since we now show when a collection is out-of-date or dirty, it makes sense to allow the user to take action to correct these states without having to leave the editor.

diff --git a/Engine/Source/Developer/CollectionManager/Private/Collection.cpp b/Engine/Source/Developer/CollectionManager/Private/Collection.cpp
--- a/Engine/Source/Developer/CollectionManager/Private/Collection.cpp
+++ b/Engine/Source/Developer/CollectionManager/Private/Collection.cpp
@@ -257,6 +257,79 @@ bool FCollection::Save(FText& OutError)
 	return bSaveSuccessful;
 }
 
+bool FCollection::Update(FText& OutError)
+{
+	if ( !ensure(SourceFilename.Len()) )
+	{
+		OutError = LOCTEXT("Error_Internal", "There was an internal error.");
+		return false;
+	}
+
+	if ( !bUseSCC )
+	{
+		// Not under SCC control, so already up-to-date
+		return true;
+	}
+
+	FScopedSlowTask SlowTask(1.0f, FText::Format(LOCTEXT("UpdatingCollection", "Updating Collection {0}"), FText::FromName(CollectionName )));
+
+	ISourceControlProvider& SourceControlProvider = ISourceControlModule::Get().GetProvider();
+	if ( !ISourceControlModule::Get().IsEnabled() )
+	{
+		OutError = LOCTEXT("Error_SCCDisabled", "Source control is not enabled. Enable source control in the preferences menu.");
+		return false;
+	}
+
+	if ( !SourceControlProvider.IsAvailable() )
+	{
+		OutError = LOCTEXT("Error_SCCNotAvailable", "Source control is currently not available. Check your connection and try again.");
+		return false;
+	}
+
+	const FString AbsoluteFilename = FPaths::ConvertRelativePathToFull(SourceFilename);
+	FSourceControlStatePtr SourceControlState = SourceControlProvider.GetState(AbsoluteFilename, EStateCacheUsage::ForceUpdate);
+
+	// If not at the head revision, sync up
+	if (SourceControlState.IsValid() && !SourceControlState->IsCurrent())
+	{
+		if ( SourceControlProvider.Execute(ISourceControlOperation::Create<FSync>(), AbsoluteFilename) == ECommandResult::Failed )
+		{
+			// Could not sync up with the head revision
+			OutError = FText::Format(LOCTEXT("Error_SCCSync", "Failed to sync collection '{0}' to the head revision."), FText::FromName(CollectionName));
+			return false;
+		}
+
+		// Check to see if the file exists at the head revision
+		if ( IFileManager::Get().FileExists(*SourceFilename) )
+		{
+			// File found! Load it and merge with our local changes
+			FText LoadErrorText;
+			FCollection NewCollection(SourceFilename, false);
+			if ( !NewCollection.Load(LoadErrorText) )
+			{
+				// Failed to load the head revision file so it isn't safe to delete it
+				OutError = FText::Format(LOCTEXT("Error_SCCBadHead", "Failed to load the collection '{0}' at the head revision. {1}"), FText::FromName(CollectionName), LoadErrorText);
+				return false;
+			}
+
+			// Loaded the head revision, now merge up so the files are in a consistent state
+			MergeWithCollection(NewCollection);
+		}
+
+		// Make sure we get a fresh state from the server
+		SourceControlState = SourceControlProvider.GetState(AbsoluteFilename, EStateCacheUsage::ForceUpdate);
+
+		// Got an updated version?
+		if (SourceControlState.IsValid() && !SourceControlState->IsCurrent())
+		{
+			OutError = FText::Format(LOCTEXT("Error_SCCNotCurrent", "Collection '{0}' is not at head revision after sync."), FText::FromName(CollectionName));
+			return false;
+		}
+	}
+
+	return true;
+}
+
 bool FCollection::DeleteSourceFile(FText& OutError)
 {
 	bool bSuccessfullyDeleted = false;
@@ -345,6 +418,43 @@ bool FCollection::IsRedirectorInCollection(FName ObjectPath) const
 	return DiskSnapshot.ObjectSet.Contains(ObjectPath);
 }
 
+FCollectionStatusInfo FCollection::GetStatusInfo() const
+{
+	FCollectionStatusInfo StatusInfo;
+
+	StatusInfo.bIsDirty = IsDirty();
+	StatusInfo.bIsEmpty = ObjectSet.Num() == 0;
+	StatusInfo.bUseSCC  = bUseSCC;
+
+	StatusInfo.NumObjects = ObjectSet.Num();
+
+	if (bUseSCC && ISourceControlModule::Get().IsEnabled())
+	{
+		ISourceControlProvider& SourceControlProvider = ISourceControlModule::Get().GetProvider();
+		if (SourceControlProvider.IsAvailable())
+		{
+			const FString AbsoluteFilename = FPaths::ConvertRelativePathToFull(SourceFilename);
+			StatusInfo.SCCState = SourceControlProvider.GetState(AbsoluteFilename, EStateCacheUsage::Use);
+		}
+	}
+
+	return StatusInfo;
+}
+
+bool FCollection::IsDirty() const
+{
+	if (ParentCollectionGuid != DiskSnapshot.ParentCollectionGuid)
+	{
+		return true;
+	}
+	
+	TArray<FName> ObjectsAdded;
+	TArray<FName> ObjectsRemoved;
+	GetObjectDifferencesFromDisk(ObjectsAdded, ObjectsRemoved);
+
+	return ObjectsAdded.Num() != 0 || ObjectsRemoved.Num() != 0;
+}
+
 bool FCollection::IsDynamic() const
 {
 	// @todo collection Dynamic collections
@@ -462,7 +572,7 @@ void FCollection::MergeWithCollection(const FCollection& Other)
 	}
 }
 
-void FCollection::GetObjectDifferencesFromDisk(TArray<FName>& ObjectsAdded, TArray<FName>& ObjectsRemoved)
+void FCollection::GetObjectDifferencesFromDisk(TArray<FName>& ObjectsAdded, TArray<FName>& ObjectsRemoved) const
 {
 	// Find the objects that were removed since the disk list
 	for (const FName& DiskObjectName : DiskSnapshot.ObjectSet)
diff --git a/Engine/Source/Developer/CollectionManager/Private/Collection.h b/Engine/Source/Developer/CollectionManager/Private/Collection.h
--- a/Engine/Source/Developer/CollectionManager/Private/Collection.h
+++ b/Engine/Source/Developer/CollectionManager/Private/Collection.h
@@ -44,6 +44,8 @@ public:
 	bool Load(FText& OutError);
 	/** Saves this collection to SourceFilename. If false, OutError is a human readable warning depicting the error. */
 	bool Save(FText& OutError);
+	/** Updates this collection to ensure it's the latest version from source control. If false, OutError is a human readable warning depicting the error. */
+	bool Update(FText& OutError);
 	/** Deletes the source file for this collection. If false, OutError is a human readable warning depicting the error. */
 	bool DeleteSourceFile(FText& OutError);
 
@@ -62,6 +64,12 @@ public:
 	/** Returns true when the specified redirector is in the collection. Static collections only. */
 	bool IsRedirectorInCollection(FName ObjectPath) const;
 
+	/** Get the status info for this collection */
+	FCollectionStatusInfo GetStatusInfo() const;
+
+	/** Does this collection contain unsaved changes? */
+	bool IsDirty() const;
+
 	/** Returns true if the collection contains rules instead of a flat list */
 	bool IsDynamic() const;
 
@@ -101,7 +109,7 @@ private:
 	/** Merges the assets from the specified collection with this collection */
 	void MergeWithCollection(const FCollection& Other);
 	/** Gets the object differences between collections A and B */
-	void GetObjectDifferencesFromDisk(TArray<FName>& ObjectsAdded, TArray<FName>& ObjectsRemoved);
+	void GetObjectDifferencesFromDisk(TArray<FName>& ObjectsAdded, TArray<FName>& ObjectsRemoved) const;
 	/** Checks the shared collection out from source control so it may be saved. If false, OutError is a human readable warning depicting the error. */
 	bool CheckoutCollection(FText& OutError);
 	/** Checks the shared collection in to source control after it is saved. If false, OutError is a human readable warning depicting the error. */
diff --git a/Engine/Source/Developer/CollectionManager/Private/CollectionManager.cpp b/Engine/Source/Developer/CollectionManager/Private/CollectionManager.cpp
--- a/Engine/Source/Developer/CollectionManager/Private/CollectionManager.cpp
+++ b/Engine/Source/Developer/CollectionManager/Private/CollectionManager.cpp
@@ -838,7 +838,62 @@
 	return false;
 }
 
-bool FCollectionManager::IsCollectionEmpty(FName CollectionName, ECollectionShareType::Type ShareType) const
+bool FCollectionManager::SaveCollection(FName CollectionName, ECollectionShareType::Type ShareType)
 {
 	if (!ensure(ShareType < ECollectionShareType::CST_All))
 	{
+		// Bad share type
+		LastError = LOCTEXT("Error_Internal", "There was an internal error.");
+		return false;
+	}
+
+	const FCollectionNameType CollectionKey(CollectionName, ShareType);
+	const TSharedRef<FCollection>* const CollectionRefPtr = CachedCollections.Find(CollectionKey);
+	if (CollectionRefPtr)
+	{
+		FCollectionStatusInfo StatusInfo = (*CollectionRefPtr)->GetStatusInfo();
+
+		const bool bNeedsSave = StatusInfo.bIsDirty || (StatusInfo.SCCState.IsValid() && StatusInfo.SCCState->IsModified());
+		if (!bNeedsSave)
+		{
+			// No changes - nothing to save
+			return true;
+		}
+
+		if ((*CollectionRefPtr)->Save(LastError))
+		{
+			CollectionUpdatedEvent.Broadcast(CollectionKey);
+			return true;
+		}
+	}
+
+	return false;
+}
+
+bool FCollectionManager::UpdateCollection(FName CollectionName, ECollectionShareType::Type ShareType)
+{
+	if (!ensure(ShareType < ECollectionShareType::CST_All))
+	{
+		// Bad share type
+		LastError = LOCTEXT("Error_Internal", "There was an internal error.");
+		return false;
+	}
+
+	const FCollectionNameType CollectionKey(CollectionName, ShareType);
+	const TSharedRef<FCollection>* const CollectionRefPtr = CachedCollections.Find(CollectionKey);
+	if (CollectionRefPtr)
+	{
+		if ((*CollectionRefPtr)->Update(LastError))
+		{
+			CollectionUpdatedEvent.Broadcast(CollectionKey);
+			return true;
+		}
+	}
+
+	return false;
+}
+
+bool FCollectionManager::GetCollectionStatusInfo(FName CollectionName, ECollectionShareType::Type ShareType, FCollectionStatusInfo& OutStatusInfo) const
+{
+	if (!ensure(ShareType < ECollectionShareType::CST_All))
+	{
@@ -851,7 +906,8 @@
 	const TSharedRef<FCollection>* const CollectionRefPtr = CachedCollections.Find(CollectionKey);
 	if (CollectionRefPtr)
 	{
-		return (*CollectionRefPtr)->IsEmpty();
+		OutStatusInfo = (*CollectionRefPtr)->GetStatusInfo();
+		return true;
 	}
 
 	return true;
@@ -944,10 +1000,28 @@
 		}
 	}
 
-	// Notify the rename for each redirected object
+	TArray<FCollectionNameType> UpdatedCollections;
+
+	TArray<FName> AddedObjects;
+	AddedObjects.Reserve(ObjectsToRename.Num());
+
+	TArray<FName> RemovedObjects;
+	RemovedObjects.Reserve(ObjectsToRename.Num());
+
+	// Handle the rename for each redirected object
 	for (const auto& ObjectToRename : ObjectsToRename)
 	{
-		HandleObjectRenamed(ObjectToRename.Key, ObjectToRename.Value);
+		AddedObjects.Add(ObjectToRename.Value);
+		RemovedObjects.Add(ObjectToRename.Key);
+
+		ReplaceObjectInCollections(ObjectToRename.Key, ObjectToRename.Value, UpdatedCollections);
 	}
 
+	// Notify every collection that changed
+	for (const FCollectionNameType& UpdatedCollection : UpdatedCollections)
+	{
+		AssetsRemovedEvent.Broadcast(UpdatedCollection, RemovedObjects);
+		AssetsAddedEvent.Broadcast(UpdatedCollection, AddedObjects);
+	}
+
 	UE_LOG(LogCollectionManager, Log, TEXT( "Fixed up redirectors for %d collections in %0.6f seconds (updated %d objects)" ), CachedCollections.Num(), FPlatformTime::Seconds() - LoadStartTime, ObjectsToRename.Num());
@@ -964,15 +1038,22 @@
 
 	FTextBuilder AllErrors;
 
+	TArray<FCollectionNameType> UpdatedCollections;
+
 	// We don't have a cache for on-disk objects, so we have to do this the slower way and query each collection in turn
 	for (const auto& CachedCollection : CachedCollections)
 	{
+		const FCollectionNameType& CollectionKey = CachedCollection.Key;
 		const TSharedRef<FCollection>& Collection = CachedCollection.Value;
 
 		if (Collection->IsRedirectorInCollection(ObjectPath))
 		{
 			FText SaveError;
-			if (!Collection->Save(SaveError))
+			if (Collection->Save(SaveError))
 			{
+				UpdatedCollections.Add(CollectionKey);
+			}
+			else
+			{
 				AllErrors.AppendLine(SaveError);
 				bSavedAllCollections = false;
@@ -980,6 +1061,15 @@
 		}
 	}
 
+	TArray<FName> RemovedObjects;
+	RemovedObjects.Add(ObjectPath);
+
+	// Notify every collection that changed
+	for (const FCollectionNameType& UpdatedCollection : UpdatedCollections)
+	{
+		AssetsRemovedEvent.Broadcast(UpdatedCollection, RemovedObjects);
+	}
+
 	if (!bSavedAllCollections)
 	{
 		LastError = AllErrors.ToText();
@@ -990,62 +1080,35 @@
 
 void FCollectionManager::HandleObjectRenamed(const FName& OldObjectPath, const FName& NewObjectPath)
 {
-	// Remove the info about the collections that the old object path is within - we'll add this info back in under the new object path shortly
-	TArray<FObjectCollectionInfo> OldObjectCollectionInfos;
-	if (!CachedObjects.RemoveAndCopyValue(OldObjectPath, OldObjectCollectionInfos))
+	TArray<FCollectionNameType> UpdatedCollections;
+	ReplaceObjectInCollections(OldObjectPath, NewObjectPath, UpdatedCollections);
+
+	TArray<FName> AddedObjects;
+	AddedObjects.Add(NewObjectPath);
+
+	TArray<FName> RemovedObjects;
+	RemovedObjects.Add(OldObjectPath);
+
+	// Notify every collection that changed
+	for (const FCollectionNameType& UpdatedCollection : UpdatedCollections)
 	{
-		return;
-	}
-
-	// Replace this object reference in all collections that use it, and update our object cache
-	auto& NewObjectCollectionInfos = CachedObjects.FindOrAdd(NewObjectPath);
-	for (const FObjectCollectionInfo& OldObjectCollectionInfo : OldObjectCollectionInfos)
-	{
-		if ((OldObjectCollectionInfo.Reason & ECollectionRecursionFlags::Self) != 0)
-		{
-			// The old object is contained directly within this collection (rather than coming from a parent or child collection), so update the object reference
-			const TSharedRef<FCollection>* const CollectionRefPtr = CachedCollections.Find(OldObjectCollectionInfo.CollectionKey);
-			if (CollectionRefPtr)
-			{
-				(*CollectionRefPtr)->RemoveObjectFromCollection(OldObjectPath);
-				(*CollectionRefPtr)->AddObjectToCollection(NewObjectPath);
-			}
-		}
-
-		// Merge the collection references for the old object with any collection references that already exist for the new object
-		FObjectCollectionInfo* ObjectInfoPtr = NewObjectCollectionInfos.FindByPredicate([&](const FObjectCollectionInfo& InCollectionInfo) { return InCollectionInfo.CollectionKey == OldObjectCollectionInfo.CollectionKey; });
-		if (ObjectInfoPtr)
-		{
-			ObjectInfoPtr->Reason |= OldObjectCollectionInfo.Reason;
-		}
-		else
-		{
-			NewObjectCollectionInfos.Add(OldObjectCollectionInfo);
-		}
+		AssetsRemovedEvent.Broadcast(UpdatedCollection, RemovedObjects);
+		AssetsAddedEvent.Broadcast(UpdatedCollection, AddedObjects);
 	}
 }
 
 void FCollectionManager::HandleObjectDeleted(const FName& ObjectPath)
 {
-	// Remove the info about the collections that the old object path is within
-	TArray<FObjectCollectionInfo> ObjectCollectionInfos;
-	if (!CachedObjects.RemoveAndCopyValue(ObjectPath, ObjectCollectionInfos))
+	TArray<FCollectionNameType> UpdatedCollections;
+	RemoveObjectFromCollections(ObjectPath, UpdatedCollections);
+
+	TArray<FName> RemovedObjects;
+	RemovedObjects.Add(ObjectPath);
+
+	// Notify every collection that changed
+	for (const FCollectionNameType& UpdatedCollection : UpdatedCollections)
 	{
-		return;
-	}
-
-	// Remove this object reference from all collections that use it
-	for (const FObjectCollectionInfo& ObjectCollectionInfo : ObjectCollectionInfos)
-	{
-		if ((ObjectCollectionInfo.Reason & ECollectionRecursionFlags::Self) != 0)
-		{
-			// The object is contained directly within this collection (rather than coming from a parent or child collection), so remove the object reference
-			const TSharedRef<FCollection>* const CollectionRefPtr = CachedCollections.Find(ObjectCollectionInfo.CollectionKey);
-			if (CollectionRefPtr)
-			{
-				(*CollectionRefPtr)->RemoveObjectFromCollection(ObjectPath);
-			}
-		}
+		AssetsRemovedEvent.Broadcast(UpdatedCollection, RemovedObjects);
 	}
 }
 
@@ -1213,6 +1276,71 @@
 	return CachedCollections.Remove(CollectionKey) > 0;
 }
 
+void FCollectionManager::RemoveObjectFromCollections(const FName& ObjectPath, TArray<FCollectionNameType>& OutUpdatedCollections)
+{
+	// Remove the info about the collections that the old object path is within
+	TArray<FObjectCollectionInfo> ObjectCollectionInfos;
+	if (!CachedObjects.RemoveAndCopyValue(ObjectPath, ObjectCollectionInfos))
+	{
+		return;
+	}
+
+	// Remove this object reference from all collections that use it
+	for (const FObjectCollectionInfo& ObjectCollectionInfo : ObjectCollectionInfos)
+	{
+		if ((ObjectCollectionInfo.Reason & ECollectionRecursionFlags::Self) != 0)
+		{
+			// The object is contained directly within this collection (rather than coming from a parent or child collection), so remove the object reference
+			const TSharedRef<FCollection>* const CollectionRefPtr = CachedCollections.Find(ObjectCollectionInfo.CollectionKey);
+			if (CollectionRefPtr)
+			{
+				OutUpdatedCollections.AddUnique(ObjectCollectionInfo.CollectionKey);
+
+				(*CollectionRefPtr)->RemoveObjectFromCollection(ObjectPath);
+			}
+		}
+	}
+}
+
+void FCollectionManager::ReplaceObjectInCollections(const FName& OldObjectPath, const FName& NewObjectPath, TArray<FCollectionNameType>& OutUpdatedCollections)
+{
+	// Remove the info about the collections that the old object path is within - we'll add this info back in under the new object path shortly
+	TArray<FObjectCollectionInfo> OldObjectCollectionInfos;
+	if (!CachedObjects.RemoveAndCopyValue(OldObjectPath, OldObjectCollectionInfos))
+	{
+		return;
+	}
+
+	// Replace this object reference in all collections that use it, and update our object cache
+	auto& NewObjectCollectionInfos = CachedObjects.FindOrAdd(NewObjectPath);
+	for (const FObjectCollectionInfo& OldObjectCollectionInfo : OldObjectCollectionInfos)
+	{
+		if ((OldObjectCollectionInfo.Reason & ECollectionRecursionFlags::Self) != 0)
+		{
+			// The old object is contained directly within this collection (rather than coming from a parent or child collection), so update the object reference
+			const TSharedRef<FCollection>* const CollectionRefPtr = CachedCollections.Find(OldObjectCollectionInfo.CollectionKey);
+			if (CollectionRefPtr)
+			{
+				OutUpdatedCollections.AddUnique(OldObjectCollectionInfo.CollectionKey);
+
+				(*CollectionRefPtr)->RemoveObjectFromCollection(OldObjectPath);
+				(*CollectionRefPtr)->AddObjectToCollection(NewObjectPath);
+			}
+		}
+
+		// Merge the collection references for the old object with any collection references that already exist for the new object
+		FObjectCollectionInfo* ObjectInfoPtr = NewObjectCollectionInfos.FindByPredicate([&](const FObjectCollectionInfo& InCollectionInfo) { return InCollectionInfo.CollectionKey == OldObjectCollectionInfo.CollectionKey; });
+		if (ObjectInfoPtr)
+		{
+			ObjectInfoPtr->Reason |= OldObjectCollectionInfo.Reason;
+		}
+		else
+		{
+			NewObjectCollectionInfos.Add(OldObjectCollectionInfo);
+		}
+	}
+}
+
 void FCollectionManager::RecursionHelper_DoWork(const FCollectionNameType& InCollectionKey, const ECollectionRecursionFlags::Flags InRecursionMode, const FRecursiveWorkerFunc& InWorkerFunc) const
 {
 	if ((InRecursionMode & ECollectionRecursionFlags::Self) && InWorkerFunc(InCollectionKey, ECollectionRecursionFlags::Self) == ERecursiveWorkerFlowControl::Stop)
diff --git a/Engine/Source/Developer/CollectionManager/Private/CollectionManager.h b/Engine/Source/Developer/CollectionManager/Private/CollectionManager.h
--- a/Engine/Source/Developer/CollectionManager/Private/CollectionManager.h
+++ b/Engine/Source/Developer/CollectionManager/Private/CollectionManager.h
@@ -36,7 +36,9 @@ public:
 	virtual bool RemoveFromCollection(FName CollectionName, ECollectionShareType::Type ShareType, FName ObjectPath) override;
 	virtual bool RemoveFromCollection(FName CollectionName, ECollectionShareType::Type ShareType, const TArray<FName>& ObjectPaths, int32* OutNumRemoved = nullptr) override;
 	virtual bool EmptyCollection(FName CollectionName, ECollectionShareType::Type ShareType) override;
-	virtual bool IsCollectionEmpty(FName CollectionName, ECollectionShareType::Type ShareType) const override;
+	virtual bool SaveCollection(FName CollectionName, ECollectionShareType::Type ShareType) override;
+	virtual bool UpdateCollection(FName CollectionName, ECollectionShareType::Type ShareType) override;
+	virtual bool GetCollectionStatusInfo(FName CollectionName, ECollectionShareType::Type ShareType, FCollectionStatusInfo& OutStatusInfo) const override;
 	virtual bool IsObjectInCollection(FName ObjectPath, FName CollectionName, ECollectionShareType::Type ShareType, ECollectionRecursionFlags::Flags RecursionMode = ECollectionRecursionFlags::Self) const override;
 	virtual bool IsValidParentCollection(FName CollectionName, ECollectionShareType::Type ShareType, FName ParentCollectionName, ECollectionShareType::Type ParentShareType) const override;
 	virtual FText GetLastError() const override { return LastError; }
@@ -69,6 +71,10 @@ public:
 	DECLARE_DERIVED_EVENT( FCollectionManager, ICollectionManager::FCollectionReparentedEvent, FCollectionReparentedEvent );
 	virtual FCollectionReparentedEvent& OnCollectionReparented() override { return CollectionReparentedEvent; }
 
+	/** Event for when collections is updated, or otherwise changed and we can't tell exactly how (eg, after updating from source control and merging) */
+	DECLARE_DERIVED_EVENT( FCollectionManager, ICollectionManager::FCollectionUpdatedEvent, FCollectionUpdatedEvent );
+	virtual FCollectionUpdatedEvent& OnCollectionUpdated() override { return CollectionUpdatedEvent; }
+
 private:
 	/** Loads all collection files from disk */
 	void LoadCollections();
@@ -91,6 +97,12 @@ private:
 	/** Removes a collection from the lookup maps */
 	bool RemoveCollection(const TSharedRef<FCollection>& CollectionRef, ECollectionShareType::Type ShareType);
 
+	/** Removes an object from any collections that contain it */
+	void RemoveObjectFromCollections(const FName& ObjectPath, TArray<FCollectionNameType>& OutUpdatedCollections);
+
+	/** Replaces an object with another in any collections that contain it */
+	void ReplaceObjectInCollections(const FName& OldObjectPath, const FName& NewObjectPath, TArray<FCollectionNameType>& OutUpdatedCollections);
+
 	enum class ERecursiveWorkerFlowControl : uint8
 	{
 		Stop,
@@ -158,6 +170,9 @@ private:
 	/** Event for when collections are re-parented */
 	FCollectionReparentedEvent CollectionReparentedEvent;
 
+	/** Event for when collections are updated, or otherwise changed and we can't tell exactly how (eg, after updating from source control and merging) */
+	FCollectionUpdatedEvent CollectionUpdatedEvent;
+
 	/** Event for when collections are created */
 	FCollectionCreatedEvent CollectionCreatedEvent;
 
diff --git a/Engine/Source/Developer/CollectionManager/Public/CollectionManagerTypes.h b/Engine/Source/Developer/CollectionManager/Public/CollectionManagerTypes.h
--- a/Engine/Source/Developer/CollectionManager/Public/CollectionManagerTypes.h
+++ b/Engine/Source/Developer/CollectionManager/Public/CollectionManagerTypes.h
@@ -2,6 +2,8 @@
 
 #pragma once
 
+#include "ISourceControlState.h"
+
 struct ECollectionShareType
 {
 	enum Type
@@ -172,6 +174,25 @@ struct FCollectionNameType
 	ECollectionShareType::Type Type;
 };
 
+/** The status information for a collection */
+struct FCollectionStatusInfo
+{
+	/** True if the collection has unsaved changes */
+	bool bIsDirty;
+
+	/** True if the collection is empty */
+	bool bIsEmpty;
+
+	/** True if the collection uses SCC */
+	bool bUseSCC;
+
+	/** Number of objects within the collection (static collections only) */
+	int32 NumObjects;
+
+	/** The current source control state of the collection. Only filled in for those collections that are under source control, but may also be empty if the SCC is unavailable (see bUseSCC) */
+	FSourceControlStatePtr SCCState;
+};
+
 class ICollectionRedirectorFollower
 {
 public:
diff --git a/Engine/Source/Developer/CollectionManager/Public/ICollectionManager.h b/Engine/Source/Developer/CollectionManager/Public/ICollectionManager.h
--- a/Engine/Source/Developer/CollectionManager/Public/ICollectionManager.h
+++ b/Engine/Source/Developer/CollectionManager/Public/ICollectionManager.h
@@ -131,18 +131,44 @@ public:
 	/**
 	  * Removes all assets from the specified collection.
 	  *
-	  * @param CollectionName The collection from which to remove the asset
+	  * @param CollectionName The collection to empty
 	  * @param ShareType The way the collection is shared.
-	  * @return true if the remove was successful. If false, GetLastError will return a human readable string description of the error.
+	  * @return true if the clear was successful. If false, GetLastError will return a human readable string description of the error.
 	  */
 	virtual bool EmptyCollection(FName CollectionName, ECollectionShareType::Type ShareType) = 0;
 
 	/**
-	  * @param CollectionName The collection from which to remove the asset
+	  * Save the collection (if dirty) and check it into source control (if under SCC control)
+	  *
+	  * Note: Generally you won't need to save collections manually as the collection manager takes care of that as objects and added/removed, etc. 
+	  *	      however, you may want to manually save a collection if a previous save attempt failed (and you've since corrected the issue), or if 
+	  *       the collection contains redirected object references that you'd like to save to disk.
+	  *
+	  * @param CollectionName The collection to save
 	  * @param ShareType The way the collection is shared.
-	  * @return true if the collection is empty.
+	  * @return true if the save was successful. If false, GetLastError will return a human readable string description of the error.
 	  */
-	virtual bool IsCollectionEmpty(FName CollectionName, ECollectionShareType::Type ShareType) const = 0;
+	virtual bool SaveCollection(FName CollectionName, ECollectionShareType::Type ShareType) = 0;
+
+	/**
+	  * Update the collection to make sure it's using the latest version from source control (if under SCC control)
+	  *
+	  * Note: Generally you won't need to update collections manually as the collection manager takes care of that as collections are saved to disk.
+	  *
+	  * @param CollectionName The collection to update
+	  * @param ShareType The way the collection is shared.
+	  * @return true if the update was successful. If false, GetLastError will return a human readable string description of the error.
+	  */
+	virtual bool UpdateCollection(FName CollectionName, ECollectionShareType::Type ShareType) = 0;
+
+	/**
+	  * Gets the status info for the specified collection
+	  *
+	  * @param CollectionName The collection to get the status info for
+	  * @param ShareType The way the collection is shared.
+	  * @return true if the status info was filled in. If false, GetLastError will return a human readable string description of the error.
+	  */
+	virtual bool GetCollectionStatusInfo(FName CollectionName, ECollectionShareType::Type ShareType, FCollectionStatusInfo& OutStatusInfo) const = 0;
 
 	/**
 	 * Check to see if the given object exists in the given collection
@@ -209,4 +235,8 @@ public:
 	/** Event for when collections are re-parented (params: Collection, OldParent, NewParent) */
 	DECLARE_EVENT_ThreeParams( ICollectionManager, FCollectionReparentedEvent, const FCollectionNameType&, const TOptional<FCollectionNameType>&, const TOptional<FCollectionNameType>& );
 	virtual FCollectionReparentedEvent& OnCollectionReparented() = 0;
+
+	/** Event for when collections is updated, or otherwise changed and we can't tell exactly how (eg, after updating from source control and merging) */
+	DECLARE_EVENT_OneParam( ICollectionManager, FCollectionUpdatedEvent, const FCollectionNameType& );
+	virtual FCollectionUpdatedEvent& OnCollectionUpdated() = 0;
 };
diff --git a/Engine/Source/Editor/ContentBrowser/Private/CollectionContextMenu.cpp b/Engine/Source/Editor/ContentBrowser/Private/CollectionContextMenu.cpp
--- a/Engine/Source/Editor/ContentBrowser/Private/CollectionContextMenu.cpp
+++ b/Engine/Source/Editor/ContentBrowser/Private/CollectionContextMenu.cpp
@@ -48,44 +48,93 @@ TSharedPtr<SWidget> FCollectionContextMenu::MakeCollectionTreeContextMenu(TShare
 
 	TArray<TSharedPtr<FCollectionItem>> SelectedCollections = CollectionView.Pin()->CollectionTreePtr->GetSelectedItems();
 
+	bool bAnyManagedBySCC = false;
+	bool bAnyNeedSCCUpdate = false;
+	bool bAnyNeedSave = false;
+		
+	for (int32 CollectionIdx = 0; CollectionIdx < SelectedCollections.Num(); ++CollectionIdx)
+	{
+		bAnyManagedBySCC |= SelectedCollections[CollectionIdx]->CollectionType != ECollectionShareType::CST_Local;
+		bAnyNeedSCCUpdate |= SelectedCollections[CollectionIdx]->CurrentStatus == ECollectionItemStatus::IsOutOfDate;
+		bAnyNeedSave |= SelectedCollections[CollectionIdx]->CurrentStatus == ECollectionItemStatus::HasLocalChanges;
+
+		if (bAnyManagedBySCC && bAnyNeedSCCUpdate && bAnyNeedSave)
+		{
+			// Found collections to turn all options on, break now
+			break;
+		}
+	}
+
 	MenuBuilder.BeginSection("CollectionOptions", LOCTEXT("CollectionListOptionsMenuHeading", "Collection Options"));
 	{
-		if ( SelectedCollections.Num() == 1 )
+		const bool bHasSingleSelectedCollection = SelectedCollections.Num() == 1;
+
+		// New... (submenu)
+		MenuBuilder.AddSubMenu(
+			LOCTEXT("NewChildCollection", "New..."),
+			LOCTEXT("NewChildCollectionTooltip", "Create child a collection."),
+			FNewMenuDelegate::CreateRaw( this, &FCollectionContextMenu::MakeNewCollectionSubMenu, TOptional<FCollectionNameType>( FCollectionNameType( SelectedCollections[0]->CollectionName, SelectedCollections[0]->CollectionType ) ) ),
+			FUIAction(
+				FExecuteAction(),
+				FCanExecuteAction::CreateLambda( [=]{ return bHasSingleSelectedCollection; } )
+				),
+			NAME_None,
+			EUserInterfaceActionType::Button
+			);
+
+		// Rename
+		MenuBuilder.AddMenuEntry( FGenericCommands::Get().Rename, NAME_None, LOCTEXT("RenameCollection", "Rename"), LOCTEXT("RenameCollectionTooltip", "Rename this collection."));
+
+		// Set Share Type
+		MenuBuilder.AddSubMenu(
+			LOCTEXT("SetCollectionShareType", "Set Share Type"),
+			LOCTEXT("SetCollectionShareTypeTooltip", "Change the share type of this collection."),
+			FNewMenuDelegate::CreateRaw( this, &FCollectionContextMenu::MakeCollectionShareTypeSubMenu ),
+			FUIAction(
+				FExecuteAction(),
+				FCanExecuteAction::CreateLambda( [=]{ return bHasSingleSelectedCollection; } )
+				),
+			NAME_None,
+			EUserInterfaceActionType::Button
+			);
+
+		// If any colors have already been set, display color options as a sub menu
+		if ( CollectionViewUtils::HasCustomColors() )
 		{
-			// New... (submenu)
+			// Set Color (submenu)
 			MenuBuilder.AddSubMenu(
-				LOCTEXT("NewChildCollection", "New..."),
-				LOCTEXT("NewChildCollectionTooltip", "Create child a collection."),
-				FNewMenuDelegate::CreateRaw( this, &FCollectionContextMenu::MakeNewCollectionSubMenu, TOptional<FCollectionNameType>( FCollectionNameType( SelectedCollections[0]->CollectionName, SelectedCollections[0]->CollectionType ) ) )
+				LOCTEXT("SetColor", "Set Color"),
+				LOCTEXT("SetCollectionColorTooltip", "Sets the color this collection should appear as."),
+				FNewMenuDelegate::CreateRaw( this, &FCollectionContextMenu::MakeSetColorSubMenu )
 				);
 		}
-
-		// Only add rename/delete if at least one collection is selected
-		if ( SelectedCollections.Num() )
+		else
 		{
-			bool bAnyManagedBySCC = false;
-		
-			for (int32 CollectionIdx = 0; CollectionIdx < SelectedCollections.Num(); ++CollectionIdx)
-			{
-				if ( SelectedCollections[CollectionIdx]->CollectionType != ECollectionShareType::CST_Local )
-				{
-					bAnyManagedBySCC = true;
-					break;
-				}
-			}
+			// Set Color
+			MenuBuilder.AddMenuEntry(
+				LOCTEXT("SetColor", "Set Color"),
+				LOCTEXT("SetCollectionColorTooltip", "Sets the color this collection should appear as."),
+				FSlateIcon(),
+				FUIAction( FExecuteAction::CreateSP( this, &FCollectionContextMenu::ExecutePickColor ) )
+				);
+		}
+	}
+	MenuBuilder.EndSection();
 
-			if ( SelectedCollections.Num() == 1 )
-			{
-				// Share Type
-				MenuBuilder.AddSubMenu(
-					LOCTEXT("SetCollectionShareType", "Share Type"),
-					LOCTEXT("SetCollectionShareTypeTooltip", "Change the share type of this collection."),
-					FNewMenuDelegate::CreateRaw( this, &FCollectionContextMenu::MakeCollectionShareTypeSubMenu )
-					);
-
-				// Rename
-				MenuBuilder.AddMenuEntry( FGenericCommands::Get().Rename, NAME_None, LOCTEXT("RenameCollection", "Rename"), LOCTEXT("RenameCollectionTooltip", "Rename this collection."));
-			}
+	if ( SelectedCollections.Num() > 0 )
+	{
+		MenuBuilder.BeginSection("CollectionBulkOperations", LOCTEXT("CollectionListBulkOperationsMenuHeading", "Bulk Operations"));
+		{
+			// Save
+			MenuBuilder.AddMenuEntry(
+				LOCTEXT("SaveCollection", "Save"),
+				LOCTEXT("SaveCollectionTooltip", "Save this collection."),
+				FSlateIcon(),
+				FUIAction(
+					FExecuteAction::CreateSP( this, &FCollectionContextMenu::ExecuteSaveCollection ),
+					FCanExecuteAction::CreateLambda( [=]{ return bAnyNeedSave; } )
+					)
+				);
 
 			// Delete
 			MenuBuilder.AddMenuEntry(
@@ -98,32 +147,19 @@ TSharedPtr<SWidget> FCollectionContextMenu::MakeCollectionTreeContextMenu(TShare
 					)
 				);
 
-			if ( SelectedCollections.Num() == 1 )
-			{
-				// If any colors have already been set, display color options as a sub menu
-				if ( CollectionViewUtils::HasCustomColors() )
-				{
-					// Set Color (submenu)
-					MenuBuilder.AddSubMenu(
-						LOCTEXT("SetColor", "Set Color"),
-						LOCTEXT("SetCollectionColorTooltip", "Sets the color this collection should appear as."),
-						FNewMenuDelegate::CreateRaw( this, &FCollectionContextMenu::MakeSetColorSubMenu )
-						);
-				}
-				else
-				{
-					// Set Color
-					MenuBuilder.AddMenuEntry(
-						LOCTEXT("SetColor", "Set Color"),
-						LOCTEXT("SetCollectionColorTooltip", "Sets the color this collection should appear as."),
-						FSlateIcon(),
-						FUIAction( FExecuteAction::CreateSP( this, &FCollectionContextMenu::ExecutePickColor ) )
-						);
-				}
-			}
+			// Update
+			MenuBuilder.AddMenuEntry(
+				LOCTEXT("UpdateCollection", "Update"),
+				LOCTEXT("UpdateCollectionTooltip", "Update this collection to make sure it's using the latest version from source control."),
+				FSlateIcon(),
+				FUIAction(
+					FExecuteAction::CreateSP( this, &FCollectionContextMenu::ExecuteUpdateCollection ),
+					FCanExecuteAction::CreateLambda( [=]{ return bAnyNeedSCCUpdate; } )
+					)
+				);
 		}
+		MenuBuilder.EndSection();
 	}
-	MenuBuilder.EndSection();
 	
 	return MenuBuilder.MakeWidget();
 }
@@ -339,6 +375,40 @@ void FCollectionContextMenu::ExecuteRenameCollection()
 	CollectionView.Pin()->RenameCollectionItem(SelectedCollections[0]);
 }
 
+void FCollectionContextMenu::ExecuteUpdateCollection()
+{
+	if ( !ensure(CollectionView.IsValid()) )
+	{
+		return;
+	}
+
+	FCollectionManagerModule& CollectionManagerModule = FCollectionManagerModule::GetModule();
+
+	TArray<TSharedPtr<FCollectionItem>> SelectedCollections = CollectionView.Pin()->CollectionTreePtr->GetSelectedItems();
+
+	for (const auto& SelectedCollection : SelectedCollections)
+	{
+		CollectionManagerModule.Get().UpdateCollection(SelectedCollection->CollectionName, SelectedCollection->CollectionType);
+	}
+}
+
+void FCollectionContextMenu::ExecuteSaveCollection()
+{
+	if ( !ensure(CollectionView.IsValid()) )
+	{
+		return;
+	}
+
+	FCollectionManagerModule& CollectionManagerModule = FCollectionManagerModule::GetModule();
+
+	TArray<TSharedPtr<FCollectionItem>> SelectedCollections = CollectionView.Pin()->CollectionTreePtr->GetSelectedItems();
+
+	for (const auto& SelectedCollection : SelectedCollections)
+	{
+		CollectionManagerModule.Get().SaveCollection(SelectedCollection->CollectionName, SelectedCollection->CollectionType);
+	}
+}
+
 void FCollectionContextMenu::ExecuteDestroyCollection()
 {
 	if ( !ensure(CollectionView.IsValid()) )
diff --git a/Engine/Source/Editor/ContentBrowser/Private/CollectionContextMenu.h b/Engine/Source/Editor/ContentBrowser/Private/CollectionContextMenu.h
--- a/Engine/Source/Editor/ContentBrowser/Private/CollectionContextMenu.h
+++ b/Engine/Source/Editor/ContentBrowser/Private/CollectionContextMenu.h
@@ -40,6 +40,12 @@ private:
 	/** Handler for when "Rename Collection" is selected */
 	void ExecuteRenameCollection();
 
+	/** Handler for when "Update Collection" is selected */
+	void ExecuteUpdateCollection();
+
+	/** Handler for when "Save Collection" is selected */
+	void ExecuteSaveCollection();
+
 	/** Handler for when "Destroy Collection" is selected */
 	void ExecuteDestroyCollection();
 
diff --git a/Engine/Source/Editor/ContentBrowser/Private/CollectionViewTypes.h b/Engine/Source/Editor/ContentBrowser/Private/CollectionViewTypes.h
--- a/Engine/Source/Editor/ContentBrowser/Private/CollectionViewTypes.h
+++ b/Engine/Source/Editor/ContentBrowser/Private/CollectionViewTypes.h
@@ -2,6 +2,31 @@
 
 #pragma once
 
+/** Status states that a collection item can be in */
+enum class ECollectionItemStatus : uint8
+{
+	/** The collection is up-to-date in source control, and isn't empty */
+	IsUpToDateAndPopulated,
+		
+	/** The collection is up-to-date in source control, but is empty */
+	IsUpToDateAndEmpty,
+
+	/** The collection is out-of-date in source control */
+	IsOutOfDate,
+
+	/** The collection is checked out by another source control user, so can't be modified at this time */
+	IsCheckedOutByAnotherUser,
+
+	/** The collection is conflicted in source control, so can't be modified at this time */
+	IsConflicted,
+
+	/** The collection is under source control but the SCC provider is currently unavailable, so can't be modified at this time */
+	IsMissingSCCProvider,
+
+	/** The collection has local changes that either haven't been saved, or haven't been committed to source control */
+	HasLocalChanges,
+};
+
 /** A list item representing a collection */
 struct FCollectionItem
 {
@@ -31,6 +56,9 @@ struct FCollectionItem
 	/** If true, this item will be created the next time it is renamed */
 	bool bNewCollection;
 
+	/** Current status of this collection item */
+	ECollectionItemStatus CurrentStatus;
+
 	/** Broadcasts whenever renaming a collection item is requested */
 	DECLARE_MULTICAST_DELEGATE( FRenamedRequestEvent )
 
@@ -45,5 +73,6 @@ struct FCollectionItem
 		, ChildCollections()
 		, bRenaming(false)
 		, bNewCollection(false)
+		, CurrentStatus(ECollectionItemStatus::IsUpToDateAndEmpty)
 	{}
 };
diff --git a/Engine/Source/Editor/ContentBrowser/Private/SCollectionView.cpp b/Engine/Source/Editor/ContentBrowser/Private/SCollectionView.cpp
--- a/Engine/Source/Editor/ContentBrowser/Private/SCollectionView.cpp
+++ b/Engine/Source/Editor/ContentBrowser/Private/SCollectionView.cpp
@@ -10,6 +10,7 @@
 #include "ObjectTools.h"
 #include "SourcesViewWidgets.h"
 #include "ContentBrowserModule.h"
+#include "ISourceControlModule.h"
 #include "SExpandableArea.h"
 #include "SSearchBox.h"
 
@@ -66,12 +67,19 @@ void SCollectionView::Construct( const FArguments& InArgs )
 	bAllowRightClickMenu = InArgs._AllowRightClickMenu;
 	bAllowCollectionDrag = InArgs._AllowCollectionDrag;
 	bDraggedOver = false;
+	bQueueSCCRefresh = true;
 
 	FCollectionManagerModule& CollectionManagerModule = FCollectionManagerModule::GetModule();
 	CollectionManagerModule.Get().OnCollectionCreated().AddSP( this, &SCollectionView::HandleCollectionCreated );
 	CollectionManagerModule.Get().OnCollectionRenamed().AddSP( this, &SCollectionView::HandleCollectionRenamed );
 	CollectionManagerModule.Get().OnCollectionReparented().AddSP( this, &SCollectionView::HandleCollectionReparented );
 	CollectionManagerModule.Get().OnCollectionDestroyed().AddSP( this, &SCollectionView::HandleCollectionDestroyed );
+	CollectionManagerModule.Get().OnCollectionUpdated().AddSP( this, &SCollectionView::HandleCollectionUpdated );
+	CollectionManagerModule.Get().OnAssetsAdded().AddSP( this, &SCollectionView::HandleAssetsAddedToCollection );
+	CollectionManagerModule.Get().OnAssetsRemoved().AddSP( this, &SCollectionView::HandleAssetsRemovedFromCollection );
+
+	ISourceControlModule::Get().RegisterProviderChanged(FSourceControlProviderChanged::FDelegate::CreateSP(this, &SCollectionView::HandleSourceControlProviderChanged));
+	SourceControlStateChangedDelegateHandle = ISourceControlModule::Get().GetProvider().RegisterSourceControlStateChanged_Handle(FSourceControlStateChanged::FDelegate::CreateSP(this, &SCollectionView::HandleSourceControlStateChanged));
 
 	Commands = TSharedPtr< FUICommandList >(new FUICommandList);
 	CollectionContextMenu = MakeShareable(new FCollectionContextMenu( SharedThis(this) ));
@@ -262,6 +270,102 @@ void SCollectionView::HandleCollectionDestroyed( const FCollectionNameType& Coll
 	UpdateCollectionItems();
 }
 
+void SCollectionView::HandleCollectionUpdated( const FCollectionNameType& Collection )
+{
+	TSharedPtr<FCollectionItem> CollectionItemToUpdate = AvailableCollections.FindRef(Collection);
+	if (CollectionItemToUpdate.IsValid())
+	{
+		bQueueSCCRefresh = true;
+		UpdateCollectionItemStatus(CollectionItemToUpdate.ToSharedRef());
+	}
+}
+
+void SCollectionView::HandleAssetsAddedToCollection( const FCollectionNameType& Collection, const TArray<FName>& AssetsAdded )
+{
+	HandleCollectionUpdated(Collection);
+}
+
+void SCollectionView::HandleAssetsRemovedFromCollection( const FCollectionNameType& Collection, const TArray<FName>& AssetsRemoved )
+{
+	HandleCollectionUpdated(Collection);
+}
+
+void SCollectionView::HandleSourceControlProviderChanged(ISourceControlProvider& OldProvider, ISourceControlProvider& NewProvider)
+{
+	OldProvider.UnregisterSourceControlStateChanged_Handle(SourceControlStateChangedDelegateHandle);
+	SourceControlStateChangedDelegateHandle = NewProvider.RegisterSourceControlStateChanged_Handle(FSourceControlStateChanged::FDelegate::CreateSP(this, &SCollectionView::HandleSourceControlStateChanged));
+	
+	bQueueSCCRefresh = true;
+	HandleSourceControlStateChanged();
+}
+
+void SCollectionView::HandleSourceControlStateChanged()
+{
+	// Update the status of each collection
+	for (const auto& AvailableCollectionInfo : AvailableCollections)
+	{
+		UpdateCollectionItemStatus(AvailableCollectionInfo.Value.ToSharedRef());
+	}
+}
+
+void SCollectionView::UpdateCollectionItemStatus( const TSharedRef<FCollectionItem>& CollectionItem )
+{
+	FCollectionManagerModule& CollectionManagerModule = FCollectionManagerModule::GetModule();
+
+	TOptional<ECollectionItemStatus> NewStatus;
+
+	FCollectionStatusInfo StatusInfo;
+	if (CollectionManagerModule.Get().GetCollectionStatusInfo(CollectionItem->CollectionName, CollectionItem->CollectionType, StatusInfo))
+	{
+		// Test the SCC state first as this should take priority when reporting the status back to the user
+		if (StatusInfo.bUseSCC)
+		{
+			if (StatusInfo.SCCState.IsValid() && StatusInfo.SCCState->IsSourceControlled())
+			{
+				if (StatusInfo.SCCState->IsCheckedOutOther())
+				{
+					NewStatus = ECollectionItemStatus::IsCheckedOutByAnotherUser;
+				}
+				else if (StatusInfo.SCCState->IsConflicted())
+				{
+					NewStatus = ECollectionItemStatus::IsConflicted;
+				}
+				else if (!StatusInfo.SCCState->IsCurrent())
+				{
+					NewStatus = ECollectionItemStatus::IsOutOfDate;
+				}
+				else if (StatusInfo.SCCState->IsModified())
+				{
+					NewStatus = ECollectionItemStatus::HasLocalChanges;
+				}
+			}
+			else
+			{
+				NewStatus = ECollectionItemStatus::IsMissingSCCProvider;
+			}
+		}
+
+		// Not set by the SCC status, so check just use the local state
+		if (!NewStatus.IsSet())
+		{
+			if (StatusInfo.bIsDirty)
+			{
+				NewStatus = ECollectionItemStatus::HasLocalChanges;
+			}
+			else if (StatusInfo.bIsEmpty)
+			{
+				NewStatus = ECollectionItemStatus::IsUpToDateAndEmpty;
+			}
+			else
+			{
+				NewStatus = ECollectionItemStatus::IsUpToDateAndPopulated;
+			}
+		}
+	}
+
+	CollectionItem->CurrentStatus = NewStatus.Get(ECollectionItemStatus::IsUpToDateAndEmpty);
+}
+
 void SCollectionView::UpdateCollectionItems()
 {
 	struct FGatherCollectionItems
@@ -299,9 +403,11 @@ void SCollectionView::UpdateCollectionItems()
 					continue;
 				}
 
-				TSharedPtr<FCollectionItem> CollectionItem = MakeShareable(new FCollectionItem(Collection.Name, Collection.Type));
+				TSharedRef<FCollectionItem> CollectionItem = MakeShareable(new FCollectionItem(Collection.Name, Collection.Type));
 				OutAvailableCollections.Add(Collection, CollectionItem);
 
+				SCollectionView::UpdateCollectionItemStatus(CollectionItem);
+
 				if (InParentCollectionItem.IsValid())
 				{
 					// Fixup the parent and child pointers
@@ -347,6 +453,8 @@ void SCollectionView::UpdateCollectionItems()
 	// Restore selection and expansion
 	SetSelectedCollections(SelectedCollections, false);
 	SetExpandedCollections(ExpandedCollections);
+
+	bQueueSCCRefresh = true;
 }
 
 void SCollectionView::UpdateFilteredCollectionItems()
@@ -609,6 +717,36 @@ void SCollectionView::LoadSettings(const FString& IniFilename, const FString& In
 	}
 }
 
+void SCollectionView::Tick( const FGeometry& AllottedGeometry, const double InCurrentTime, const float InDeltaTime )
+{
+	SCompoundWidget::Tick(AllottedGeometry, InCurrentTime, InDeltaTime);
+
+	if (bQueueSCCRefresh && ISourceControlModule::Get().IsEnabled())
+	{
+		bQueueSCCRefresh = false;
+
+		FCollectionManagerModule& CollectionManagerModule = FCollectionManagerModule::GetModule();
+
+		TArray<FString> CollectionFilesToRefresh;
+		for (const auto& AvailableCollectionInfo : AvailableCollections)
+		{
+			FCollectionStatusInfo StatusInfo;
+			if (CollectionManagerModule.Get().GetCollectionStatusInfo(AvailableCollectionInfo.Value->CollectionName, AvailableCollectionInfo.Value->CollectionType, StatusInfo))
+			{
+				if (StatusInfo.bUseSCC && StatusInfo.SCCState.IsValid() && StatusInfo.SCCState->IsSourceControlled())
+				{
+					CollectionFilesToRefresh.Add(StatusInfo.SCCState->GetFilename());
+				}
+			}
+		}
+
+		if (CollectionFilesToRefresh.Num() > 0)
+		{
+			ISourceControlModule::Get().QueueStatusUpdate(CollectionFilesToRefresh);
+		}
+	}
+}
+
 FReply SCollectionView::OnKeyDown( const FGeometry& MyGeometry, const FKeyEvent& InKeyEvent )
 {
 	if( Commands->ProcessCommandBindings( InKeyEvent ) )
diff --git a/Engine/Source/Editor/ContentBrowser/Private/SCollectionView.h b/Engine/Source/Editor/ContentBrowser/Private/SCollectionView.h
--- a/Engine/Source/Editor/ContentBrowser/Private/SCollectionView.h
+++ b/Engine/Source/Editor/ContentBrowser/Private/SCollectionView.h
@@ -70,6 +70,7 @@ public:
 	/** Loads any settings to config that should be persistent between editor sessions */
 	void LoadSettings(const FString& IniFilename, const FString& IniSection, const FString& SettingsString);
 
+	virtual void Tick( const FGeometry& AllottedGeometry, const double InCurrentTime, const float InDeltaTime ) override;
 	virtual FReply OnKeyDown( const FGeometry& MyGeometry, const FKeyEvent& InKeyEvent ) override;
 	virtual void OnDragEnter( const FGeometry& MyGeometry, const FDragDropEvent& DragDropEvent ) override;
 	virtual void OnDragLeave( const FDragDropEvent& DragDropEvent ) override;
@@ -171,6 +172,24 @@ private:
 	/** Handles an on collection destroyed event */
 	void HandleCollectionDestroyed( const FCollectionNameType& Collection );
 
+	/** Handles an on collection updated event */
+	void HandleCollectionUpdated( const FCollectionNameType& Collection );
+
+	/** Handles assets being added to a collection */
+	void HandleAssetsAddedToCollection( const FCollectionNameType& Collection, const TArray<FName>& AssetsAdded );
+
+	/** Handles assets being removed from a collection */
+	void HandleAssetsRemovedFromCollection( const FCollectionNameType& Collection, const TArray<FName>& AssetsRemoved );
+
+	/** Handles the source control provider changing */
+	void HandleSourceControlProviderChanged(class ISourceControlProvider& OldProvider, class ISourceControlProvider& NewProvider);
+
+	/** Handles the source control state changing */
+	void HandleSourceControlStateChanged();
+
+	/** Handles updating the status of the given collection item */
+	static void UpdateCollectionItemStatus( const TSharedRef<FCollectionItem>& CollectionItem );
+
 	/** Updates the collections shown in the tree view */
 	void UpdateCollectionItems();
 
@@ -260,4 +279,10 @@ private:
 	 * We keep a weak pointer to this so we can tell if that drag and drop is still ongoing, and if so, what collections it affects 
 	 */
 	TWeakPtr<class FCollectionDragDropOp> CurrentCollectionDragDropOp;
+
+	/** Delegate handle for the HandleSourceControlStateChanged function callback */
+	FDelegateHandle SourceControlStateChangedDelegateHandle;
+
+	/** True if we should queue an SCC refresh for the collections on the next Tick */
+	bool bQueueSCCRefresh;
 };
\ No newline at end of file
diff --git a/Engine/Source/Editor/ContentBrowser/Private/SourcesViewWidgets.cpp b/Engine/Source/Editor/ContentBrowser/Private/SourcesViewWidgets.cpp
--- a/Engine/Source/Editor/ContentBrowser/Private/SourcesViewWidgets.cpp
+++ b/Engine/Source/Editor/ContentBrowser/Private/SourcesViewWidgets.cpp
@@ -366,6 +366,17 @@ void SCollectionTreeItem::Construct( const FArguments& InArgs )
 				.IsSelected( InArgs._IsSelected )
 				.IsReadOnly( InArgs._IsReadOnly )
 			]
+
+			+SHorizontalBox::Slot()
+			.AutoWidth()
+			.VAlign(VAlign_Center)
+			.Padding(2, 0, 2, 0)
+			[
+				SNew(SImage)
+				.Image(FEditorStyle::GetBrush("ContentBrowser.CollectionStatus"))
+				.ColorAndOpacity(this, &SCollectionTreeItem::GetCollectionStatusColor)
+				.ToolTipText(this, &SCollectionTreeItem::GetCollectionStatusToolTipText)
+			]
 		]
 	];
 
@@ -534,5 +545,74 @@ const FSlateBrush* SCollectionTreeItem::GetBorderImage() const
 	return bDraggedOver ? FEditorStyle::GetBrush("Menu.Background") : FEditorStyle::GetBrush("NoBorder");
 }
 
+FSlateColor SCollectionTreeItem::GetCollectionStatusColor() const
+{
+	TSharedPtr<FCollectionItem> CollectionItemPtr = CollectionItem.Pin();
+
+	if (CollectionItemPtr.IsValid())
+	{
+		switch(CollectionItemPtr->CurrentStatus)
+		{
+		case ECollectionItemStatus::IsUpToDateAndPopulated:
+			return FLinearColor(0.10616, 0.48777, 0.10616); // Green
+
+		case ECollectionItemStatus::IsUpToDateAndEmpty:
+			return FLinearColor::Gray;
+
+		case ECollectionItemStatus::IsOutOfDate:
+			return FLinearColor(0.87514, 0.42591, 0.07383); // Orange
+
+		case ECollectionItemStatus::IsCheckedOutByAnotherUser:
+		case ECollectionItemStatus::IsConflicted:
+		case ECollectionItemStatus::IsMissingSCCProvider:
+			return FLinearColor(0.70117, 0.08464, 0.07593); // Red
+
+		case ECollectionItemStatus::HasLocalChanges:
+			return FLinearColor(0.10363, 0.53564, 0.7372); // Blue
+
+		default:
+			break;
+		}
+	}
+
+	return FLinearColor::White;
+}
+
+FText SCollectionTreeItem::GetCollectionStatusToolTipText() const
+{
+	TSharedPtr<FCollectionItem> CollectionItemPtr = CollectionItem.Pin();
+
+	if (CollectionItemPtr.IsValid())
+	{
+		switch(CollectionItemPtr->CurrentStatus)
+		{
+		case ECollectionItemStatus::IsUpToDateAndPopulated:
+			return LOCTEXT("CollectionStatus_IsUpToDateAndPopulated", "Collection is up-to-date");
+
+		case ECollectionItemStatus::IsUpToDateAndEmpty:
+			return LOCTEXT("CollectionStatus_IsUpToDateAndEmpty", "Collection is empty");
+
+		case ECollectionItemStatus::IsOutOfDate:
+			return LOCTEXT("CollectionStatus_IsOutOfDate", "Collection is not at the latest revision");
+
+		case ECollectionItemStatus::IsCheckedOutByAnotherUser:
+			return LOCTEXT("CollectionStatus_IsCheckedOutByAnotherUser", "Collection is checked out by another user");
+
+		case ECollectionItemStatus::IsConflicted:
+			return LOCTEXT("CollectionStatus_IsConflicted", "Collection is conflicted - please use your external source control provider to resolve this conflict");
+
+		case ECollectionItemStatus::IsMissingSCCProvider:
+			return LOCTEXT("CollectionStatus_IsMissingSCCProvider", "Collection is missing its source control provider - please check your source control settings");
+
+		case ECollectionItemStatus::HasLocalChanges:
+			return LOCTEXT("CollectionStatus_HasLocalChanges", "Collection has local unsaved or uncomitted changes");
+
+		default:
+			break;
+		}
+	}
+
+	return FText::GetEmpty();
+}
 
 #undef LOCTEXT_NAMESPACE
diff --git a/Engine/Source/Editor/ContentBrowser/Private/SourcesViewWidgets.h b/Engine/Source/Editor/ContentBrowser/Private/SourcesViewWidgets.h
--- a/Engine/Source/Editor/ContentBrowser/Private/SourcesViewWidgets.h
+++ b/Engine/Source/Editor/ContentBrowser/Private/SourcesViewWidgets.h
@@ -240,6 +240,12 @@ private:
 	/** Returns the image for the border around this item. Used for drag/drop operations */
 	const FSlateBrush* GetBorderImage() const;
 
+	/** Get the color to use for the collection item status */
+	FSlateColor GetCollectionStatusColor() const;
+
+	/** Get the tooltip correspinding to the current collection item status */
+	FText GetCollectionStatusToolTipText() const;
+
 private:
 	/** A shared pointer to the parent widget. */
 	TSharedPtr<SWidget> ParentWidget;
diff --git a/Engine/Source/Editor/EditorStyle/Private/SlateEditorStyle.cpp b/Engine/Source/Editor/EditorStyle/Private/SlateEditorStyle.cpp
--- a/Engine/Source/Editor/EditorStyle/Private/SlateEditorStyle.cpp
+++ b/Engine/Source/Editor/EditorStyle/Private/SlateEditorStyle.cpp
@@ -5580,6 +5580,7 @@ void FSlateEditorStyle::FStyle::SetupContentBrowserStyle()
 		Set( "ContentBrowser.Private", new IMAGE_BRUSH( "ContentBrowser/Content_Private_12x", Icon12x12 ) );
 		Set( "ContentBrowser.Private.Small", new IMAGE_BRUSH( "ContentBrowser/Content_Private_16x", Icon16x16 ) );
 		Set( "ContentBrowser.Private.Large", new IMAGE_BRUSH( "ContentBrowser/Content_Private_64x", Icon64x64 ) );
+		Set( "ContentBrowser.CollectionStatus", new IMAGE_BRUSH( "/Icons/CollectionStatus_8x", Icon8x8 ) );
 	}
 #endif // #if WITH_EDITOR
 }
