Fixed some places in the collection manager that needed to update a cache

diff --git a/Engine/Source/Developer/CollectionManager/Private/Collection.cpp b/Engine/Source/Developer/CollectionManager/Private/Collection.cpp
--- a/Engine/Source/Developer/CollectionManager/Private/Collection.cpp
+++ b/Engine/Source/Developer/CollectionManager/Private/Collection.cpp
@@ -312,6 +312,9 @@ bool FCollection::Update(FText& OutError)
 				return false;
 			}
 
+			// Update the info about what the collection looks like on disk
+			DiskSnapshot.TakeSnapshot(NewCollection);
+
 			// Loaded the head revision, now merge up so the files are in a consistent state
 			MergeWithCollection(NewCollection);
 		}
@@ -654,6 +657,9 @@ bool FCollection::CheckoutCollection(FText& OutError)
 				return false;
 			}
 
+			// Update the info about what the collection looks like on disk
+			DiskSnapshot.TakeSnapshot(NewCollection);
+
 			// Loaded the head revision, now merge up so the files are in a consistent state
 			MergeWithCollection(NewCollection);
 		}
@@ -963,6 +969,9 @@ bool FCollection::DeleteFromSourceControl(FText& OutError)
 			return false;
 		}
 
+		// Update the info about what the collection looks like on disk
+		DiskSnapshot.TakeSnapshot(NewCollection);
+
 		// Loaded the head revision, now merge up so the files are in a consistent state
 		MergeWithCollection(NewCollection);
 
diff --git a/Engine/Source/Developer/CollectionManager/Private/CollectionManager.cpp b/Engine/Source/Developer/CollectionManager/Private/CollectionManager.cpp
--- a/Engine/Source/Developer/CollectionManager/Private/CollectionManager.cpp
+++ b/Engine/Source/Developer/CollectionManager/Private/CollectionManager.cpp
@@ -885,6 +885,9 @@ bool FCollectionManager::UpdateCollection(FName CollectionName, ECollectionShare
 	{
 		if ((*CollectionRefPtr)->Update(LastError))
 		{
+			RebuildCachedHierarchy();
+			RebuildCachedObjects();
+
 			CollectionUpdatedEvent.Broadcast(CollectionKey);
 			return true;
 		}
