Made the collection manager caching lazy

This used to maintain an up-to-date cache as soon as a change happened. This lead to extra complexity in some cases, and often would keep caches up-to-date that weren't currently even being used.

Adding, removing, or updating a collection will now just mark the relevant parts of the cache as dirty so that they can be updated (using a well defined, safe, full update) the next time something requests them.

diff --git a/Engine/Source/Developer/CollectionManager/Private/CollectionManager.cpp b/Engine/Source/Developer/CollectionManager/Private/CollectionManager.cpp
--- a/Engine/Source/Developer/CollectionManager/Private/CollectionManager.cpp
+++ b/Engine/Source/Developer/CollectionManager/Private/CollectionManager.cpp
@@ -4,7 +4,220 @@
 
 #define LOCTEXT_NAMESPACE "CollectionManager"
 
+FCollectionManagerCache::FCollectionManagerCache(FAvailableCollectionsMap& InAvailableCollections)
+	: AvailableCollections(InAvailableCollections)
+{
+	bIsCachedCollectionNamesFromGuidsDirty = true;
+	bIsCachedObjectsDirty = true;
+	bIsCachedHierarchyDirty = true;
+}
+
+void FCollectionManagerCache::HandleCollectionAdded()
+{
+	bIsCachedCollectionNamesFromGuidsDirty = true;
+}
+
+void FCollectionManagerCache::HandleCollectionRemoved()
+{
+	bIsCachedCollectionNamesFromGuidsDirty = true;
+	bIsCachedObjectsDirty = true;
+	bIsCachedHierarchyDirty = true;
+}
+
+void FCollectionManagerCache::HandleCollectionChanged()
+{
+	bIsCachedObjectsDirty = true;
+	bIsCachedHierarchyDirty = true;
+}
+
+const FGuidToCollectionNamesMap& FCollectionManagerCache::GetCachedCollectionNamesFromGuids() const
+{
+	if (bIsCachedCollectionNamesFromGuidsDirty)
+	{
+		CachedCollectionNamesFromGuids_Internal.Reset();
+		bIsCachedCollectionNamesFromGuidsDirty = false;
+
+		const double CacheStartTime = FPlatformTime::Seconds();
+
+		for (const auto& AvailableCollection : AvailableCollections)
+		{
+			const FCollectionNameType& CollectionKey = AvailableCollection.Key;
+			const TSharedRef<FCollection>& Collection = AvailableCollection.Value;
+
+			CachedCollectionNamesFromGuids_Internal.Add(Collection->GetCollectionGuid(), CollectionKey);
+		}
+
+		UE_LOG(LogCollectionManager, Log, TEXT("Rebuilt the GUID cache for %d collections in %0.6f seconds"), AvailableCollections.Num(), FPlatformTime::Seconds() - CacheStartTime);
+	}
+
+	return CachedCollectionNamesFromGuids_Internal;
+}
+
+const FCollectionObjectsMap& FCollectionManagerCache::GetCachedObjects() const
+{
+	if (bIsCachedObjectsDirty)
+	{
+		CachedObjects_Internal.Reset();
+		bIsCachedObjectsDirty = false;
+
+		const double CacheStartTime = FPlatformTime::Seconds();
+
+		for (const auto& AvailableCollection : AvailableCollections)
+		{
+			const FCollectionNameType& CollectionKey = AvailableCollection.Key;
+			const TSharedRef<FCollection>& Collection = AvailableCollection.Value;
+
+			TArray<FName> ObjectsInCollection;
+			Collection->GetObjectsInCollection(ObjectsInCollection);
+
+			if (ObjectsInCollection.Num() > 0)
+			{
+				auto RebuildCachedObjectsWorker = [&](const FCollectionNameType& InCollectionKey, ECollectionRecursionFlags::Flag InReason) -> ERecursiveWorkerFlowControl
+				{
+					// The worker reason will tell us why this collection is being processed (eg, because it is a parent of the collection we told it to DoWork on),
+					// however, the reason this object exists in that parent collection is because a child collection contains it, and this is the reason we need
+					// to put into the FObjectCollectionInfo, since that's what we'll test against later when we do the "do my children contain this object"? test
+					// That's why we flip the reason logic here...
+					ECollectionRecursionFlags::Flag ReasonObjectInCollection = InReason;
+					switch (InReason)
+					{
+					case ECollectionRecursionFlags::Parents:
+						ReasonObjectInCollection = ECollectionRecursionFlags::Children;
+						break;
+					case ECollectionRecursionFlags::Children:
+						ReasonObjectInCollection = ECollectionRecursionFlags::Parents;
+						break;
+					default:
+						break;
+					}
+
+					for (const FName& ObjectPath : ObjectsInCollection)
+					{
+						auto& ObjectCollectionInfos = CachedObjects_Internal.FindOrAdd(ObjectPath);
+						FObjectCollectionInfo* ObjectInfoPtr = ObjectCollectionInfos.FindByPredicate([&](const FObjectCollectionInfo& InCollectionInfo) { return InCollectionInfo.CollectionKey == InCollectionKey; });
+						if (ObjectInfoPtr)
+						{
+							ObjectInfoPtr->Reason |= ReasonObjectInCollection;
+						}
+						else
+						{
+							ObjectCollectionInfos.Add(FObjectCollectionInfo(InCollectionKey, ReasonObjectInCollection));
+						}
+					}
+					return ERecursiveWorkerFlowControl::Continue;
+				};
+
+				// Recursively process all collections so that they know they contain these objects (and why!)
+				RecursionHelper_DoWork(CollectionKey, ECollectionRecursionFlags::All, RebuildCachedObjectsWorker);
+			}
+		}
+
+		UE_LOG(LogCollectionManager, Log, TEXT("Rebuilt the object cache for %d collections in %0.6f seconds (found %d objects)"), AvailableCollections.Num(), FPlatformTime::Seconds() - CacheStartTime, CachedObjects_Internal.Num());
+	}
+
+	return CachedObjects_Internal;
+}
+
+const FCollectionHierarchyMap& FCollectionManagerCache::GetCachedHierarchy() const
+{
+	if (bIsCachedHierarchyDirty)
+	{
+		CachedHierarchy_Internal.Reset();
+		bIsCachedHierarchyDirty = false;
+
+		const FGuidToCollectionNamesMap& CachedCollectionNamesFromGuids = GetCachedCollectionNamesFromGuids();
+
+		const double CacheStartTime = FPlatformTime::Seconds();
+
+		for (const auto& AvailableCollection : AvailableCollections)
+		{
+			const FCollectionNameType& CollectionKey = AvailableCollection.Key;
+			const TSharedRef<FCollection>& Collection = AvailableCollection.Value;
+
+			// Make sure this is a known parent GUID before adding it to the map
+			const FGuid& ParentCollectionGuid = Collection->GetParentCollectionGuid();
+			if (CachedCollectionNamesFromGuids.Contains(ParentCollectionGuid))
+			{
+				auto& CollectionChildren = CachedHierarchy_Internal.FindOrAdd(ParentCollectionGuid);
+				CollectionChildren.AddUnique(Collection->GetCollectionGuid());
+			}
+		}
+
+		UE_LOG(LogCollectionManager, Log, TEXT("Rebuilt the hierarchy cache for %d collections in %0.6f seconds"), AvailableCollections.Num(), FPlatformTime::Seconds() - CacheStartTime);
+	}
+
+	return CachedHierarchy_Internal;
+}
+
+void FCollectionManagerCache::RecursionHelper_DoWork(const FCollectionNameType& InCollectionKey, const ECollectionRecursionFlags::Flags InRecursionMode, const FRecursiveWorkerFunc& InWorkerFunc) const
+{
+	if ((InRecursionMode & ECollectionRecursionFlags::Self) && InWorkerFunc(InCollectionKey, ECollectionRecursionFlags::Self) == ERecursiveWorkerFlowControl::Stop)
+	{
+		return;
+	}
+
+	if ((InRecursionMode & ECollectionRecursionFlags::Parents) && RecursionHelper_DoWorkOnParents(InCollectionKey, InWorkerFunc) == ERecursiveWorkerFlowControl::Stop)
+	{
+		return;
+	}
+
+	if ((InRecursionMode & ECollectionRecursionFlags::Children) && RecursionHelper_DoWorkOnChildren(InCollectionKey, InWorkerFunc) == ERecursiveWorkerFlowControl::Stop)
+	{
+		return;
+	}
+}
+
+FCollectionManagerCache::ERecursiveWorkerFlowControl FCollectionManagerCache::RecursionHelper_DoWorkOnParents(const FCollectionNameType& InCollectionKey, const FRecursiveWorkerFunc& InWorkerFunc) const
+{
+	const TSharedRef<FCollection>* const CollectionRefPtr = AvailableCollections.Find(InCollectionKey);
+	if (CollectionRefPtr)
+	{
+		const FGuidToCollectionNamesMap& CachedCollectionNamesFromGuids = GetCachedCollectionNamesFromGuids();
+
+		const FCollectionNameType* const ParentCollectionKeyPtr = CachedCollectionNamesFromGuids.Find((*CollectionRefPtr)->GetParentCollectionGuid());
+		if (ParentCollectionKeyPtr)
+		{
+			if (InWorkerFunc(*ParentCollectionKeyPtr, ECollectionRecursionFlags::Parents) == ERecursiveWorkerFlowControl::Stop || RecursionHelper_DoWorkOnParents(*ParentCollectionKeyPtr, InWorkerFunc) == ERecursiveWorkerFlowControl::Stop)
+			{
+				return ERecursiveWorkerFlowControl::Stop;
+			}
+		}
+	}
+
+	return ERecursiveWorkerFlowControl::Continue;
+}
+
+FCollectionManagerCache::ERecursiveWorkerFlowControl FCollectionManagerCache::RecursionHelper_DoWorkOnChildren(const FCollectionNameType& InCollectionKey, const FRecursiveWorkerFunc& InWorkerFunc) const
+{
+	const TSharedRef<FCollection>* const CollectionRefPtr = AvailableCollections.Find(InCollectionKey);
+	if (CollectionRefPtr)
+	{
+		const FCollectionHierarchyMap& CachedHierarchy = GetCachedHierarchy();
+
+		const TArray<FGuid>* const ChildCollectionGuids = CachedHierarchy.Find((*CollectionRefPtr)->GetCollectionGuid());
+		if (ChildCollectionGuids)
+		{
+			for (const FGuid& ChildCollectionGuid : *ChildCollectionGuids)
+			{
+				const FGuidToCollectionNamesMap& CachedCollectionNamesFromGuids = GetCachedCollectionNamesFromGuids();
+
+				const FCollectionNameType* const ChildCollectionKeyPtr = CachedCollectionNamesFromGuids.Find(ChildCollectionGuid);
+				if (ChildCollectionKeyPtr)
+				{
+					if (InWorkerFunc(*ChildCollectionKeyPtr, ECollectionRecursionFlags::Children) == ERecursiveWorkerFlowControl::Stop || RecursionHelper_DoWorkOnChildren(*ChildCollectionKeyPtr, InWorkerFunc) == ERecursiveWorkerFlowControl::Stop)
+					{
+						return ERecursiveWorkerFlowControl::Stop;
+					}
+				}
+			}
+		}
+	}
+
+	return ERecursiveWorkerFlowControl::Continue;
+}
+
 FCollectionManager::FCollectionManager()
+	: CollectionCache(AvailableCollections)
 {
 	LastError = LOCTEXT("Error_Unknown", "None");
 
@@ -23,24 +236,24 @@ FCollectionManager::~FCollectionManager()
 
 bool FCollectionManager::HasCollections() const
 {
-	return CachedCollections.Num() > 0;
+	return AvailableCollections.Num() > 0;
 }
 
 void FCollectionManager::GetCollections(TArray<FCollectionNameType>& OutCollections) const
 {
-	OutCollections.Reserve(CachedCollections.Num());
-	for (const auto& CachedCollection : CachedCollections)
+	OutCollections.Reserve(AvailableCollections.Num());
+	for (const auto& AvailableCollection : AvailableCollections)
 	{
-		const FCollectionNameType& CollectionKey = CachedCollection.Key;
+		const FCollectionNameType& CollectionKey = AvailableCollection.Key;
 		OutCollections.Add(CollectionKey);
 	}
 }
 
 void FCollectionManager::GetCollectionNames(ECollectionShareType::Type ShareType, TArray<FName>& CollectionNames) const
 {
-	for (const auto& CachedCollection : CachedCollections)
+	for (const auto& AvailableCollection : AvailableCollections)
 	{
-		const FCollectionNameType& CollectionKey = CachedCollection.Key;
+		const FCollectionNameType& CollectionKey = AvailableCollection.Key;
 		if (ShareType == ECollectionShareType::CST_All || ShareType == CollectionKey.Type)
 		{
 			CollectionNames.AddUnique(CollectionKey.Name);
@@ -50,11 +263,13 @@ void FCollectionManager::GetCollectionNames(ECollectionShareType::Type ShareType
 
 void FCollectionManager::GetRootCollections(TArray<FCollectionNameType>& OutCollections) const
 {
-	OutCollections.Reserve(CachedCollections.Num());
-	for (const auto& CachedCollection : CachedCollections)
+	const FGuidToCollectionNamesMap& CachedCollectionNamesFromGuids = CollectionCache.GetCachedCollectionNamesFromGuids();
+
+	OutCollections.Reserve(AvailableCollections.Num());
+	for (const auto& AvailableCollection : AvailableCollections)
 	{
-		const FCollectionNameType& CollectionKey = CachedCollection.Key;
-		const TSharedRef<FCollection>& Collection = CachedCollection.Value;
+		const FCollectionNameType& CollectionKey = AvailableCollection.Key;
+		const TSharedRef<FCollection>& Collection = AvailableCollection.Value;
 
 		// A root collection either has no parent GUID, or a parent GUID that cannot currently be found - the check below handles both
 		if (!CachedCollectionNamesFromGuids.Contains(Collection->GetParentCollectionGuid()))
@@ -66,10 +281,12 @@ void FCollectionManager::GetRootCollections(TArray<FCollectionNameType>& OutColl
 
 void FCollectionManager::GetRootCollectionNames(ECollectionShareType::Type ShareType, TArray<FName>& CollectionNames) const
 {
-	for (const auto& CachedCollection : CachedCollections)
+	const FGuidToCollectionNamesMap& CachedCollectionNamesFromGuids = CollectionCache.GetCachedCollectionNamesFromGuids();
+
+	for (const auto& AvailableCollection : AvailableCollections)
 	{
-		const FCollectionNameType& CollectionKey = CachedCollection.Key;
-		const TSharedRef<FCollection>& Collection = CachedCollection.Value;
+		const FCollectionNameType& CollectionKey = AvailableCollection.Key;
+		const TSharedRef<FCollection>& Collection = AvailableCollection.Value;
 
 		if (ShareType == ECollectionShareType::CST_All || ShareType == CollectionKey.Type)
 		{
@@ -84,9 +301,12 @@ void FCollectionManager::GetRootCollectionNames(ECollectionShareType::Type Share
 
 void FCollectionManager::GetChildCollections(FName CollectionName, ECollectionShareType::Type ShareType, TArray<FCollectionNameType>& OutCollections) const
 {
-	auto GetChildCollectionsInternal = [this, &OutCollections](const FCollectionNameType& InCollectionKey)
+	const FGuidToCollectionNamesMap& CachedCollectionNamesFromGuids = CollectionCache.GetCachedCollectionNamesFromGuids();
+	const FCollectionHierarchyMap& CachedHierarchy = CollectionCache.GetCachedHierarchy();
+
+	auto GetChildCollectionsInternal = [&](const FCollectionNameType& InCollectionKey)
 	{
-		const TSharedRef<FCollection>* const CollectionRefPtr = CachedCollections.Find(InCollectionKey);
+		const TSharedRef<FCollection>* const CollectionRefPtr = AvailableCollections.Find(InCollectionKey);
 		if (CollectionRefPtr)
 		{
 			const auto* ChildCollectionGuids = CachedHierarchy.Find((*CollectionRefPtr)->GetCollectionGuid());
@@ -120,9 +340,12 @@ void FCollectionManager::GetChildCollections(FName CollectionName, ECollectionSh
 
 void FCollectionManager::GetChildCollectionNames(FName CollectionName, ECollectionShareType::Type ShareType, ECollectionShareType::Type ChildShareType, TArray<FName>& CollectionNames) const
 {
-	auto GetChildCollectionsInternal = [this, &ChildShareType, &CollectionNames](const FCollectionNameType& InCollectionKey)
+	const FGuidToCollectionNamesMap& CachedCollectionNamesFromGuids = CollectionCache.GetCachedCollectionNamesFromGuids();
+	const FCollectionHierarchyMap& CachedHierarchy = CollectionCache.GetCachedHierarchy();
+
+	auto GetChildCollectionsInternal = [&](const FCollectionNameType& InCollectionKey)
 	{
-		const TSharedRef<FCollection>* const CollectionRefPtr = CachedCollections.Find(InCollectionKey);
+		const TSharedRef<FCollection>* const CollectionRefPtr = AvailableCollections.Find(InCollectionKey);
 		if (CollectionRefPtr)
 		{
 			const auto* ChildCollectionGuids = CachedHierarchy.Find((*CollectionRefPtr)->GetCollectionGuid());
@@ -156,9 +379,11 @@ void FCollectionManager::GetChildCollectionNames(FName CollectionName, ECollecti
 
 TOptional<FCollectionNameType> FCollectionManager::GetParentCollection(FName CollectionName, ECollectionShareType::Type ShareType) const
 {
-	const TSharedRef<FCollection>* const CollectionRefPtr = CachedCollections.Find(FCollectionNameType(CollectionName, ShareType));
+	const TSharedRef<FCollection>* const CollectionRefPtr = AvailableCollections.Find(FCollectionNameType(CollectionName, ShareType));
 	if (CollectionRefPtr)
 	{
+		const FGuidToCollectionNamesMap& CachedCollectionNamesFromGuids = CollectionCache.GetCachedCollectionNamesFromGuids();
+
 		const FCollectionNameType* const ParentCollectionKeyPtr = CachedCollectionNamesFromGuids.Find((*CollectionRefPtr)->GetParentCollectionGuid());
 		if (ParentCollectionKeyPtr)
 		{
@@ -176,7 +401,7 @@ bool FCollectionManager::CollectionExists(FName CollectionName, ECollectionShare
 		// Asked to check all share types...
 		for (int32 CacheIdx = 0; CacheIdx < ECollectionShareType::CST_All; ++CacheIdx)
 		{
-			if (CachedCollections.Contains(FCollectionNameType(CollectionName, ECollectionShareType::Type(CacheIdx))))
+			if (AvailableCollections.Contains(FCollectionNameType(CollectionName, ECollectionShareType::Type(CacheIdx))))
 			{
 				// Collection exists in at least one cache
 				return true;
@@ -188,7 +413,7 @@ bool FCollectionManager::CollectionExists(FName CollectionName, ECollectionShare
 	}
 	else
 	{
-		return CachedCollections.Contains(FCollectionNameType(CollectionName, ShareType));
+		return AvailableCollections.Contains(FCollectionNameType(CollectionName, ShareType));
 	}
 }
 
@@ -196,15 +421,15 @@ bool FCollectionManager::GetAssetsInCollection(FName CollectionName, ECollection
 {
 	bool bFoundAssets = false;
 
-	auto GetAssetsInCollectionWorker = [&](const FCollectionNameType& InCollectionKey, ECollectionRecursionFlags::Flag InReason) -> ERecursiveWorkerFlowControl
+	auto GetAssetsInCollectionWorker = [&](const FCollectionNameType& InCollectionKey, ECollectionRecursionFlags::Flag InReason) -> FCollectionManagerCache::ERecursiveWorkerFlowControl
 	{
-		const TSharedRef<FCollection>* const CollectionRefPtr = CachedCollections.Find(InCollectionKey);
+		const TSharedRef<FCollection>* const CollectionRefPtr = AvailableCollections.Find(InCollectionKey);
 		if (CollectionRefPtr)
 		{
 			(*CollectionRefPtr)->GetAssetsInCollection(AssetsPaths);
 			bFoundAssets = true;
 		}
-		return ERecursiveWorkerFlowControl::Continue;
+		return FCollectionManagerCache::ERecursiveWorkerFlowControl::Continue;
 	};
 
 	if (ShareType == ECollectionShareType::CST_All)
@@ -212,12 +437,12 @@ bool FCollectionManager::GetAssetsInCollection(FName CollectionName, ECollection
 		// Asked for all share types, find assets in the specified collection name in any cache
 		for (int32 CacheIdx = 0; CacheIdx < ECollectionShareType::CST_All; ++CacheIdx)
 		{
-			RecursionHelper_DoWork(FCollectionNameType(CollectionName, ECollectionShareType::Type(CacheIdx)), RecursionMode, GetAssetsInCollectionWorker);
+			CollectionCache.RecursionHelper_DoWork(FCollectionNameType(CollectionName, ECollectionShareType::Type(CacheIdx)), RecursionMode, GetAssetsInCollectionWorker);
 		}
 	}
 	else
 	{
-		RecursionHelper_DoWork(FCollectionNameType(CollectionName, ShareType), RecursionMode, GetAssetsInCollectionWorker);
+		CollectionCache.RecursionHelper_DoWork(FCollectionNameType(CollectionName, ShareType), RecursionMode, GetAssetsInCollectionWorker);
 	}
 
 	return bFoundAssets;
@@ -227,15 +452,15 @@ bool FCollectionManager::GetClassesInCollection(FName CollectionName, ECollectio
 {
 	bool bFoundClasses = false;
 
-	auto GetClassesInCollectionWorker = [&](const FCollectionNameType& InCollectionKey, ECollectionRecursionFlags::Flag InReason) -> ERecursiveWorkerFlowControl
+	auto GetClassesInCollectionWorker = [&](const FCollectionNameType& InCollectionKey, ECollectionRecursionFlags::Flag InReason) -> FCollectionManagerCache::ERecursiveWorkerFlowControl
 	{
-		const TSharedRef<FCollection>* const CollectionRefPtr = CachedCollections.Find(InCollectionKey);
+		const TSharedRef<FCollection>* const CollectionRefPtr = AvailableCollections.Find(InCollectionKey);
 		if (CollectionRefPtr)
 		{
 			(*CollectionRefPtr)->GetClassesInCollection(ClassPaths);
 			bFoundClasses = true;
 		}
-		return ERecursiveWorkerFlowControl::Continue;
+		return FCollectionManagerCache::ERecursiveWorkerFlowControl::Continue;
 	};
 
 	if (ShareType == ECollectionShareType::CST_All)
@@ -243,12 +468,12 @@ bool FCollectionManager::GetClassesInCollection(FName CollectionName, ECollectio
 		// Asked for all share types, find classes in the specified collection name in any cache
 		for (int32 CacheIdx = 0; CacheIdx < ECollectionShareType::CST_All; ++CacheIdx)
 		{
-			RecursionHelper_DoWork(FCollectionNameType(CollectionName, ECollectionShareType::Type(CacheIdx)), RecursionMode, GetClassesInCollectionWorker);
+			CollectionCache.RecursionHelper_DoWork(FCollectionNameType(CollectionName, ECollectionShareType::Type(CacheIdx)), RecursionMode, GetClassesInCollectionWorker);
 		}
 	}
 	else
 	{
-		RecursionHelper_DoWork(FCollectionNameType(CollectionName, ShareType), RecursionMode, GetClassesInCollectionWorker);
+		CollectionCache.RecursionHelper_DoWork(FCollectionNameType(CollectionName, ShareType), RecursionMode, GetClassesInCollectionWorker);
 	}
 
 	return bFoundClasses;
@@ -258,15 +483,15 @@ bool FCollectionManager::GetObjectsInCollection(FName CollectionName, ECollectio
 {
 	bool bFoundObjects = false;
 
-	auto GetObjectsInCollectionWorker = [&](const FCollectionNameType& InCollectionKey, ECollectionRecursionFlags::Flag InReason) -> ERecursiveWorkerFlowControl
+	auto GetObjectsInCollectionWorker = [&](const FCollectionNameType& InCollectionKey, ECollectionRecursionFlags::Flag InReason) -> FCollectionManagerCache::ERecursiveWorkerFlowControl
 	{
-		const TSharedRef<FCollection>* const CollectionRefPtr = CachedCollections.Find(InCollectionKey);
+		const TSharedRef<FCollection>* const CollectionRefPtr = AvailableCollections.Find(InCollectionKey);
 		if (CollectionRefPtr)
 		{
 			(*CollectionRefPtr)->GetObjectsInCollection(ObjectPaths);
 			bFoundObjects = true;
 		}
-		return ERecursiveWorkerFlowControl::Continue;
+		return FCollectionManagerCache::ERecursiveWorkerFlowControl::Continue;
 	};
 
 	if (ShareType == ECollectionShareType::CST_All)
@@ -274,12 +499,12 @@ bool FCollectionManager::GetObjectsInCollection(FName CollectionName, ECollectio
 		// Asked for all share types, find classes in the specified collection name in any cache
 		for (int32 CacheIdx = 0; CacheIdx < ECollectionShareType::CST_All; ++CacheIdx)
 		{
-			RecursionHelper_DoWork(FCollectionNameType(CollectionName, ECollectionShareType::Type(CacheIdx)), RecursionMode, GetObjectsInCollectionWorker);
+			CollectionCache.RecursionHelper_DoWork(FCollectionNameType(CollectionName, ECollectionShareType::Type(CacheIdx)), RecursionMode, GetObjectsInCollectionWorker);
 		}
 	}
 	else
 	{
-		RecursionHelper_DoWork(FCollectionNameType(CollectionName, ShareType), RecursionMode, GetObjectsInCollectionWorker);
+		CollectionCache.RecursionHelper_DoWork(FCollectionNameType(CollectionName, ShareType), RecursionMode, GetObjectsInCollectionWorker);
 	}
 
 	return bFoundObjects;
@@ -287,6 +512,8 @@ bool FCollectionManager::GetObjectsInCollection(FName CollectionName, ECollectio
 
 void FCollectionManager::GetCollectionsContainingObject(FName ObjectPath, ECollectionShareType::Type ShareType, TArray<FName>& OutCollectionNames, ECollectionRecursionFlags::Flags RecursionMode) const
 {
+	const FCollectionObjectsMap& CachedObjects = CollectionCache.GetCachedObjects();
+
 	const auto* ObjectCollectionInfosPtr = CachedObjects.Find(ObjectPath);
 	if (ObjectCollectionInfosPtr)
 	{
@@ -302,6 +529,8 @@ void FCollectionManager::GetCollectionsContainingObject(FName ObjectPath, EColle
 
 void FCollectionManager::GetCollectionsContainingObject(FName ObjectPath, TArray<FCollectionNameType>& OutCollections, ECollectionRecursionFlags::Flags RecursionMode) const
 {
+	const FCollectionObjectsMap& CachedObjects = CollectionCache.GetCachedObjects();
+
 	const auto* ObjectCollectionInfosPtr = CachedObjects.Find(ObjectPath);
 	if (ObjectCollectionInfosPtr)
 	{
@@ -318,6 +547,8 @@ void FCollectionManager::GetCollectionsContainingObject(FName ObjectPath, TArray
 
 void FCollectionManager::GetCollectionsContainingObjects(const TArray<FName>& ObjectPaths, TMap<FCollectionNameType, TArray<FName>>& OutCollectionsAndMatchedObjects, ECollectionRecursionFlags::Flags RecursionMode) const
 {
+	const FCollectionObjectsMap& CachedObjects = CollectionCache.GetCachedObjects();
+
 	for (const FName& ObjectPath : ObjectPaths)
 	{
 		const auto* ObjectCollectionInfosPtr = CachedObjects.Find(ObjectPath);
@@ -337,6 +568,8 @@ void FCollectionManager::GetCollectionsContainingObjects(const TArray<FName>& Ob
 
 FString FCollectionManager::GetCollectionsStringForObject(FName ObjectPath, ECollectionShareType::Type ShareType, ECollectionRecursionFlags::Flags RecursionMode) const
 {
+	const FCollectionObjectsMap& CachedObjects = CollectionCache.GetCachedObjects();
+
 	const auto* ObjectCollectionInfosPtr = CachedObjects.Find(ObjectPath);
 	if (ObjectCollectionInfosPtr)
 	{
@@ -344,10 +577,10 @@ FString FCollectionManager::GetCollectionsStringForObject(FName ObjectPath, ECol
 
 		TArray<FString> CollectionPathStrings;
 
-		auto GetCollectionsStringForObjectWorker = [&](const FCollectionNameType& InCollectionKey, ECollectionRecursionFlags::Flag InReason) -> ERecursiveWorkerFlowControl
+		auto GetCollectionsStringForObjectWorker = [&](const FCollectionNameType& InCollectionKey, ECollectionRecursionFlags::Flag InReason) -> FCollectionManagerCache::ERecursiveWorkerFlowControl
 		{
 			CollectionPathStrings.Insert(InCollectionKey.Name.ToString(), 0);
-			return ERecursiveWorkerFlowControl::Continue;
+			return FCollectionManagerCache::ERecursiveWorkerFlowControl::Continue;
 		};
 
 		for (const FObjectCollectionInfo& ObjectCollectionInfo : *ObjectCollectionInfosPtr)
@@ -355,7 +588,7 @@ FString FCollectionManager::GetCollectionsStringForObject(FName ObjectPath, ECol
 			if ((ShareType == ECollectionShareType::CST_All || ShareType == ObjectCollectionInfo.CollectionKey.Type) && (RecursionMode & ObjectCollectionInfo.Reason) != 0)
 			{
 				CollectionPathStrings.Reset();
-				RecursionHelper_DoWork(ObjectCollectionInfo.CollectionKey, ECollectionRecursionFlags::SelfAndParents, GetCollectionsStringForObjectWorker);
+				CollectionCache.RecursionHelper_DoWork(ObjectCollectionInfo.CollectionKey, ECollectionRecursionFlags::SelfAndParents, GetCollectionsStringForObjectWorker);
 
 				CollectionNameStrings.Add(FString::Join(CollectionPathStrings, TEXT("/")));
 			}
@@ -391,7 +624,7 @@ void FCollectionManager::CreateUniqueCollectionName(const FName& BaseName, EColl
 			OutCollectionName = *FString::Printf(TEXT("%s%d"), *BaseName.ToString(), IntSuffix);
 		}
 
-		CollectionAlreadyExists = CachedCollections.Contains(FCollectionNameType(OutCollectionName, ShareType));
+		CollectionAlreadyExists = AvailableCollections.Contains(FCollectionNameType(OutCollectionName, ShareType));
 		++IntSuffix;
 	}
 	while (CollectionAlreadyExists);
@@ -478,7 +711,7 @@ bool FCollectionManager::RenameCollection(FName CurrentCollectionName, ECollecti
 	const FCollectionNameType OriginalCollectionKey(CurrentCollectionName, CurrentShareType);
 	const FCollectionNameType NewCollectionKey(NewCollectionName, NewShareType);
 
-	TSharedRef<FCollection>* const CollectionRefPtr = CachedCollections.Find(OriginalCollectionKey);
+	TSharedRef<FCollection>* const CollectionRefPtr = AvailableCollections.Find(OriginalCollectionKey);
 	if (!CollectionRefPtr)
 	{
 		// The collection doesn't exist
@@ -505,10 +738,6 @@ bool FCollectionManager::RenameCollection(FName CurrentCollectionName, ECollecti
 		{
 			// Collection failed to save, remove it from the cache
 			RemoveCollection(NewCollection.ToSharedRef(), NewShareType);
-
-			// Restore the original GUID mapping since AddCollection replaced it
-			CachedCollectionNamesFromGuids.Add((*CollectionRefPtr)->GetCollectionGuid(), OriginalCollectionKey);
-
 			return false;
 		}
 	}
@@ -518,24 +747,17 @@ bool FCollectionManager::RenameCollection(FName CurrentCollectionName, ECollecti
 		if ((*CollectionRefPtr)->DeleteSourceFile(LastError))
 		{
 			RemoveCollection(*CollectionRefPtr, CurrentShareType);
-
-			// Re-add the new GUID mapping since RemoveCollection removed it
-			CachedCollectionNamesFromGuids.Add(NewCollection->GetCollectionGuid(), NewCollectionKey);
 		}
 		else
 		{
 			// Failed to remove the old collection, so remove the collection we created.
 			NewCollection->DeleteSourceFile(LastError);
 			RemoveCollection(NewCollection.ToSharedRef(), NewShareType);
-
-			// Restore the original GUID mapping since RemoveCollection removed it
-			CachedCollectionNamesFromGuids.Add((*CollectionRefPtr)->GetCollectionGuid(), OriginalCollectionKey);
-
 			return false;
 		}
 	}
 
-	RebuildCachedObjects();
+	CollectionCache.HandleCollectionChanged();
 
 	// Success
 	CollectionRenamedEvent.Broadcast(OriginalCollectionKey, NewCollectionKey);
@@ -552,7 +774,7 @@ bool FCollectionManager::ReparentCollection(FName CollectionName, ECollectionSha
 	}
 
 	const FCollectionNameType CollectionKey(CollectionName, ShareType);
-	TSharedRef<FCollection>* const CollectionRefPtr = CachedCollections.Find(CollectionKey);
+	TSharedRef<FCollection>* const CollectionRefPtr = AvailableCollections.Find(CollectionKey);
 	if (!CollectionRefPtr)
 	{
 		// The collection doesn't exist
@@ -570,7 +792,7 @@ bool FCollectionManager::ReparentCollection(FName CollectionName, ECollectionSha
 	{
 		// Find and set the new parent GUID
 		NewParentCollectionKey = FCollectionNameType(ParentCollectionName, ParentShareType);
-		TSharedRef<FCollection>* const ParentCollectionRefPtr = CachedCollections.Find(NewParentCollectionKey.GetValue());
+		TSharedRef<FCollection>* const ParentCollectionRefPtr = AvailableCollections.Find(NewParentCollectionKey.GetValue());
 		if (!ParentCollectionRefPtr)
 		{
 			// The collection doesn't exist
@@ -613,8 +835,12 @@ bool FCollectionManager::ReparentCollection(FName CollectionName, ECollectionSha
 		return false;
 	}
 
+	CollectionCache.HandleCollectionChanged();
+
 	// Find the old parent so we can notify about the change
 	{
+		const FGuidToCollectionNamesMap& CachedCollectionNamesFromGuids = CollectionCache.GetCachedCollectionNamesFromGuids();
+
 		const FCollectionNameType* const OldParentCollectionKeyPtr = CachedCollectionNamesFromGuids.Find(OldParentGuid);
 		if (OldParentCollectionKeyPtr)
 		{
@@ -622,9 +848,6 @@ bool FCollectionManager::ReparentCollection(FName CollectionName, ECollectionSha
 		}
 	}
 
-	RebuildCachedHierarchy();
-	RebuildCachedObjects();
-
 	// Success
 	CollectionReparentedEvent.Broadcast(CollectionKey, OldParentCollectionKey, NewParentCollectionKey);
 	return true;
@@ -640,7 +863,7 @@ bool FCollectionManager::DestroyCollection(FName CollectionName, ECollectionShar
 	}
 
 	const FCollectionNameType CollectionKey(CollectionName, ShareType);
-	TSharedRef<FCollection>* const CollectionRefPtr = CachedCollections.Find(CollectionKey);
+	TSharedRef<FCollection>* const CollectionRefPtr = AvailableCollections.Find(CollectionKey);
 	if (!CollectionRefPtr)
 	{
 		// The collection doesn't exist
@@ -651,10 +874,6 @@ bool FCollectionManager::DestroyCollection(FName CollectionName, ECollectionShar
 	if ((*CollectionRefPtr)->DeleteSourceFile(LastError))
 	{
 		RemoveCollection(*CollectionRefPtr, ShareType);
-
-		RebuildCachedHierarchy();
-		RebuildCachedObjects();
-
 		CollectionDestroyedEvent.Broadcast(CollectionKey);
 		return true;
 	}
@@ -687,7 +906,7 @@ bool FCollectionManager::AddToCollection(FName CollectionName, ECollectionShareT
 	}
 
 	const FCollectionNameType CollectionKey(CollectionName, ShareType);
-	TSharedRef<FCollection>* const CollectionRefPtr = CachedCollections.Find(CollectionKey);
+	TSharedRef<FCollection>* const CollectionRefPtr = AvailableCollections.Find(CollectionKey);
 	if (!CollectionRefPtr)
 	{
 		// Collection doesn't exist
@@ -714,8 +933,7 @@ bool FCollectionManager::AddToCollection(FName CollectionName, ECollectionShareT
 				*OutNumAdded = NumAdded;
 			}
 
-			RebuildCachedObjects();
-
+			CollectionCache.HandleCollectionChanged();
 			AssetsAddedEvent.Broadcast(CollectionKey, ObjectPaths);
 			return true;
 		}
@@ -759,7 +977,7 @@ bool FCollectionManager::RemoveFromCollection(FName CollectionName, ECollectionS
 	}
 
 	const FCollectionNameType CollectionKey(CollectionName, ShareType);
-	TSharedRef<FCollection>* const CollectionRefPtr = CachedCollections.Find(CollectionKey);
+	TSharedRef<FCollection>* const CollectionRefPtr = AvailableCollections.Find(CollectionKey);
 	if (!CollectionRefPtr)
 	{
 		// Collection not found
@@ -791,8 +1009,7 @@ bool FCollectionManager::RemoveFromCollection(FName CollectionName, ECollectionS
 			*OutNumRemoved = RemovedAssets.Num();
 		}
 
-		RebuildCachedObjects();
-
+		CollectionCache.HandleCollectionChanged();
 		AssetsRemovedEvent.Broadcast(CollectionKey, ObjectPaths);
 		return true;
 	}
@@ -829,13 +1046,7 @@ bool FCollectionManager::EmptyCollection(FName CollectionName, ECollectionShareT
 		return true;
 	}
 	
-	if (RemoveFromCollection(CollectionName, ShareType, ObjectPaths))
-	{
-		RebuildCachedObjects();
-		return true;
-	}
-
-	return false;
+	return RemoveFromCollection(CollectionName, ShareType, ObjectPaths);
 }
 
 bool FCollectionManager::SaveCollection(FName CollectionName, ECollectionShareType::Type ShareType)
@@ -848,7 +1059,7 @@ bool FCollectionManager::SaveCollection(FName CollectionName, ECollectionShareTy
 	}
 
 	const FCollectionNameType CollectionKey(CollectionName, ShareType);
-	const TSharedRef<FCollection>* const CollectionRefPtr = CachedCollections.Find(CollectionKey);
+	const TSharedRef<FCollection>* const CollectionRefPtr = AvailableCollections.Find(CollectionKey);
 	if (CollectionRefPtr)
 	{
 		FCollectionStatusInfo StatusInfo = (*CollectionRefPtr)->GetStatusInfo();
@@ -862,6 +1073,7 @@ bool FCollectionManager::SaveCollection(FName CollectionName, ECollectionShareTy
 
 		if ((*CollectionRefPtr)->Save(LastError))
 		{
+			CollectionCache.HandleCollectionChanged();
 			CollectionUpdatedEvent.Broadcast(CollectionKey);
 			return true;
 		}
@@ -880,14 +1092,12 @@ bool FCollectionManager::UpdateCollection(FName CollectionName, ECollectionShare
 	}
 
 	const FCollectionNameType CollectionKey(CollectionName, ShareType);
-	const TSharedRef<FCollection>* const CollectionRefPtr = CachedCollections.Find(CollectionKey);
+	const TSharedRef<FCollection>* const CollectionRefPtr = AvailableCollections.Find(CollectionKey);
 	if (CollectionRefPtr)
 	{
 		if ((*CollectionRefPtr)->Update(LastError))
 		{
-			RebuildCachedHierarchy();
-			RebuildCachedObjects();
-
+			CollectionCache.HandleCollectionChanged();
 			CollectionUpdatedEvent.Broadcast(CollectionKey);
 			return true;
 		}
@@ -906,7 +1116,7 @@ bool FCollectionManager::GetCollectionStatusInfo(FName CollectionName, ECollecti
 	}
 
 	const FCollectionNameType CollectionKey(CollectionName, ShareType);
-	const TSharedRef<FCollection>* const CollectionRefPtr = CachedCollections.Find(CollectionKey);
+	const TSharedRef<FCollection>* const CollectionRefPtr = AvailableCollections.Find(CollectionKey);
 	if (CollectionRefPtr)
 	{
 		OutStatusInfo = (*CollectionRefPtr)->GetStatusInfo();
@@ -927,17 +1137,17 @@ bool FCollectionManager::IsObjectInCollection(FName ObjectPath, FName Collection
 
 	bool bFoundObject = false;
 
-	auto IsObjectInCollectionWorker = [&](const FCollectionNameType& InCollectionKey, ECollectionRecursionFlags::Flag InReason) -> ERecursiveWorkerFlowControl
+	auto IsObjectInCollectionWorker = [&](const FCollectionNameType& InCollectionKey, ECollectionRecursionFlags::Flag InReason) -> FCollectionManagerCache::ERecursiveWorkerFlowControl
 	{
-		const TSharedRef<FCollection>* const CollectionRefPtr = CachedCollections.Find(InCollectionKey);
+		const TSharedRef<FCollection>* const CollectionRefPtr = AvailableCollections.Find(InCollectionKey);
 		if (CollectionRefPtr)
 		{
 			bFoundObject = (*CollectionRefPtr)->IsObjectInCollection(ObjectPath);
 		}
-		return (bFoundObject) ? ERecursiveWorkerFlowControl::Stop : ERecursiveWorkerFlowControl::Continue;
+		return (bFoundObject) ? FCollectionManagerCache::ERecursiveWorkerFlowControl::Stop : FCollectionManagerCache::ERecursiveWorkerFlowControl::Continue;
 	};
 
-	RecursionHelper_DoWork(FCollectionNameType(CollectionName, ShareType), RecursionMode, IsObjectInCollectionWorker);
+	CollectionCache.RecursionHelper_DoWork(FCollectionNameType(CollectionName, ShareType), RecursionMode, IsObjectInCollectionWorker);
 
 	return bFoundObject;
 }
@@ -959,7 +1169,7 @@ bool FCollectionManager::IsValidParentCollection(FName CollectionName, ECollecti
 
 	bool bValidParent = true;
 
-	auto IsValidParentCollectionWorker = [&](const FCollectionNameType& InCollectionKey, ECollectionRecursionFlags::Flag InReason) -> ERecursiveWorkerFlowControl
+	auto IsValidParentCollectionWorker = [&](const FCollectionNameType& InCollectionKey, ECollectionRecursionFlags::Flag InReason) -> FCollectionManagerCache::ERecursiveWorkerFlowControl
 	{
 		const bool bMatchesCollectionBeingReparented = (CollectionName == InCollectionKey.Name && ShareType == InCollectionKey.Type);
 		if (bMatchesCollectionBeingReparented)
@@ -968,7 +1178,7 @@ bool FCollectionManager::IsValidParentCollection(FName CollectionName, ECollecti
 			LastError = (InReason == ECollectionRecursionFlags::Self) 
 				? LOCTEXT("InvalidParent_CannotParentToSelf", "A collection cannot be parented to itself") 
 				: LOCTEXT("InvalidParent_CannotParentToChildren", "A collection cannot be parented to its children");
-			return ERecursiveWorkerFlowControl::Stop;
+			return FCollectionManagerCache::ERecursiveWorkerFlowControl::Stop;
 		}
 
 		const bool bIsValidChildType = ECollectionShareType::IsValidChildType(InCollectionKey.Type, ShareType);
@@ -976,13 +1186,13 @@ bool FCollectionManager::IsValidParentCollection(FName CollectionName, ECollecti
 		{
 			bValidParent = false;
 			LastError = FText::Format(LOCTEXT("InvalidParent_InvalidChildType", "A {0} collection cannot contain a {1} collection"), ECollectionShareType::ToText(InCollectionKey.Type), ECollectionShareType::ToText(ShareType));
-			return ERecursiveWorkerFlowControl::Stop;
+			return FCollectionManagerCache::ERecursiveWorkerFlowControl::Stop;
 		}
 
-		return ERecursiveWorkerFlowControl::Continue;
+		return FCollectionManagerCache::ERecursiveWorkerFlowControl::Continue;
 	};
 
-	RecursionHelper_DoWork(FCollectionNameType(ParentCollectionName, ParentShareType), ECollectionRecursionFlags::SelfAndParents, IsValidParentCollectionWorker);
+	CollectionCache.RecursionHelper_DoWork(FCollectionNameType(ParentCollectionName, ParentShareType), ECollectionRecursionFlags::SelfAndParents, IsValidParentCollectionWorker);
 
 	return bValidParent;
 }
@@ -994,12 +1204,15 @@ void FCollectionManager::HandleFixupRedirectors(ICollectionRedirectorFollower& I
 	TArray<TPair<FName, FName>> ObjectsToRename;
 
 	// Build up the list of redirected object into rename pairs
-	for (const auto& CachedObjectInfo : CachedObjects)
 	{
-		FName NewObjectPath;
-		if (InRedirectorFollower.FixupObject(CachedObjectInfo.Key, NewObjectPath))
+		const FCollectionObjectsMap& CachedObjects = CollectionCache.GetCachedObjects();
+		for (const auto& CachedObjectInfo : CachedObjects)
 		{
-			ObjectsToRename.Add(TPairInitializer<FName, FName>(CachedObjectInfo.Key, NewObjectPath));
+			FName NewObjectPath;
+			if (InRedirectorFollower.FixupObject(CachedObjectInfo.Key, NewObjectPath))
+			{
+				ObjectsToRename.Add(TPairInitializer<FName, FName>(CachedObjectInfo.Key, NewObjectPath));
+			}
 		}
 	}
 
@@ -1020,14 +1233,19 @@ void FCollectionManager::HandleFixupRedirectors(ICollectionRedirectorFollower& I
 		ReplaceObjectInCollections(ObjectToRename.Key, ObjectToRename.Value, UpdatedCollections);
 	}
 
-	// Notify every collection that changed
-	for (const FCollectionNameType& UpdatedCollection : UpdatedCollections)
+	if (UpdatedCollections.Num() > 0)
 	{
-		AssetsRemovedEvent.Broadcast(UpdatedCollection, RemovedObjects);
-		AssetsAddedEvent.Broadcast(UpdatedCollection, AddedObjects);
+		CollectionCache.HandleCollectionChanged();
+
+		// Notify every collection that changed
+		for (const FCollectionNameType& UpdatedCollection : UpdatedCollections)
+		{
+			AssetsRemovedEvent.Broadcast(UpdatedCollection, RemovedObjects);
+			AssetsAddedEvent.Broadcast(UpdatedCollection, AddedObjects);
+		}
 	}
 
-	UE_LOG(LogCollectionManager, Log, TEXT( "Fixed up redirectors for %d collections in %0.6f seconds (updated %d objects)" ), CachedCollections.Num(), FPlatformTime::Seconds() - LoadStartTime, ObjectsToRename.Num());
+	UE_LOG(LogCollectionManager, Log, TEXT( "Fixed up redirectors for %d collections in %0.6f seconds (updated %d objects)" ), AvailableCollections.Num(), FPlatformTime::Seconds() - LoadStartTime, ObjectsToRename.Num());
 
 	for (const auto& ObjectToRename : ObjectsToRename)
 	{
@@ -1044,10 +1262,10 @@ bool FCollectionManager::HandleRedirectorDeleted(const FName& ObjectPath)
 	TArray<FCollectionNameType> UpdatedCollections;
 
 	// We don't have a cache for on-disk objects, so we have to do this the slower way and query each collection in turn
-	for (const auto& CachedCollection : CachedCollections)
+	for (const auto& AvailableCollection : AvailableCollections)
 	{
-		const FCollectionNameType& CollectionKey = CachedCollection.Key;
-		const TSharedRef<FCollection>& Collection = CachedCollection.Value;
+		const FCollectionNameType& CollectionKey = AvailableCollection.Key;
+		const TSharedRef<FCollection>& Collection = AvailableCollection.Value;
 
 		if (Collection->IsRedirectorInCollection(ObjectPath))
 		{
@@ -1092,11 +1310,16 @@ void FCollectionManager::HandleObjectRenamed(const FName& OldObjectPath, const F
 	TArray<FName> RemovedObjects;
 	RemovedObjects.Add(OldObjectPath);
 
-	// Notify every collection that changed
-	for (const FCollectionNameType& UpdatedCollection : UpdatedCollections)
+	if (UpdatedCollections.Num() > 0)
 	{
-		AssetsRemovedEvent.Broadcast(UpdatedCollection, RemovedObjects);
-		AssetsAddedEvent.Broadcast(UpdatedCollection, AddedObjects);
+		CollectionCache.HandleCollectionChanged();
+	
+		// Notify every collection that changed
+		for (const FCollectionNameType& UpdatedCollection : UpdatedCollections)
+		{
+			AssetsRemovedEvent.Broadcast(UpdatedCollection, RemovedObjects);
+			AssetsAddedEvent.Broadcast(UpdatedCollection, AddedObjects);
+		}
 	}
 }
 
@@ -1108,17 +1331,22 @@ void FCollectionManager::HandleObjectDeleted(const FName& ObjectPath)
 	TArray<FName> RemovedObjects;
 	RemovedObjects.Add(ObjectPath);
 
-	// Notify every collection that changed
-	for (const FCollectionNameType& UpdatedCollection : UpdatedCollections)
+	if (UpdatedCollections.Num() > 0)
 	{
-		AssetsRemovedEvent.Broadcast(UpdatedCollection, RemovedObjects);
+		CollectionCache.HandleCollectionChanged();
+	
+		// Notify every collection that changed
+		for (const FCollectionNameType& UpdatedCollection : UpdatedCollections)
+		{
+			AssetsRemovedEvent.Broadcast(UpdatedCollection, RemovedObjects);
+		}
 	}
 }
 
 void FCollectionManager::LoadCollections()
 {
 	const double LoadStartTime = FPlatformTime::Seconds();
-	const int32 PrevNumCollections = CachedCollections.Num();
+	const int32 PrevNumCollections = AvailableCollections.Num();
 
 	for (int32 CacheIdx = 0; CacheIdx < ECollectionShareType::CST_All; ++CacheIdx)
 	{
@@ -1146,92 +1374,10 @@ void FCollectionManager::LoadCollections()
 		}
 	}
 
-	UE_LOG(LogCollectionManager, Log, TEXT( "Loaded %d collections in %0.6f seconds" ), CachedCollections.Num() - PrevNumCollections, FPlatformTime::Seconds() - LoadStartTime);
-
-	RebuildCachedHierarchy();
-	RebuildCachedObjects();
-}
-
-void FCollectionManager::RebuildCachedObjects()
-{
-	const double LoadStartTime = FPlatformTime::Seconds();
-
-	CachedObjects.Empty();
-
-	for (const auto& CachedCollection : CachedCollections)
-	{
-		const FCollectionNameType& CollectionKey = CachedCollection.Key;
-		const TSharedRef<FCollection>& Collection = CachedCollection.Value;
-
-		TArray<FName> ObjectsInCollection;
-		Collection->GetObjectsInCollection(ObjectsInCollection);
-
-		if (ObjectsInCollection.Num() > 0)
-		{
-			auto RebuildCachedObjectsWorker = [&](const FCollectionNameType& InCollectionKey, ECollectionRecursionFlags::Flag InReason) -> ERecursiveWorkerFlowControl
-			{
-				// The worker reason will tell us why this collection is being processed (eg, because it is a parent of the collection we told it to DoWork on),
-				// however, the reason this object exists in that parent collection is because a child collection contains it, and this is the reason we need
-				// to put into the FObjectCollectionInfo, since that's what we'll test against later when we do the "do my children contain this object"? test
-				// That's why we flip the reason logic here...
-				ECollectionRecursionFlags::Flag ReasonObjectInCollection = InReason;
-				switch (InReason)
-				{
-				case ECollectionRecursionFlags::Parents:
-					ReasonObjectInCollection = ECollectionRecursionFlags::Children;
-					break;
-				case ECollectionRecursionFlags::Children:
-					ReasonObjectInCollection = ECollectionRecursionFlags::Parents;
-					break;
-				default:
-					break;
-				}
-
-				for (const FName& ObjectPath : ObjectsInCollection)
-				{
-					auto& ObjectCollectionInfos = CachedObjects.FindOrAdd(ObjectPath);
-					FObjectCollectionInfo* ObjectInfoPtr = ObjectCollectionInfos.FindByPredicate([&](const FObjectCollectionInfo& InCollectionInfo) { return InCollectionInfo.CollectionKey == InCollectionKey; });
-					if (ObjectInfoPtr)
-					{
-						ObjectInfoPtr->Reason |= ReasonObjectInCollection;
-					}
-					else
-					{
-						ObjectCollectionInfos.Add(FObjectCollectionInfo(InCollectionKey, ReasonObjectInCollection));
-					}
-				}
-				return ERecursiveWorkerFlowControl::Continue;
-			};
-
-			// Recursively process all collections so that they know they contain these objects (and why!)
-			RecursionHelper_DoWork(CollectionKey, ECollectionRecursionFlags::All, RebuildCachedObjectsWorker);
-		}
-	}
-
-	UE_LOG(LogCollectionManager, Log, TEXT( "Rebuilt the object cache for %d collections in %0.6f seconds (found %d objects)" ), CachedCollections.Num(), FPlatformTime::Seconds() - LoadStartTime, CachedObjects.Num());
-}
-
-void FCollectionManager::RebuildCachedHierarchy()
-{
-	const double LoadStartTime = FPlatformTime::Seconds();
-
-	CachedHierarchy.Empty();
+	// AddCollection is assumed to be adding an empty collection, so also notify that collection cache that the collection has "changed" since loaded collections may not always be empty
+	CollectionCache.HandleCollectionChanged();
 
-	for (const auto& CachedCollection : CachedCollections)
-	{
-		const FCollectionNameType& CollectionKey = CachedCollection.Key;
-		const TSharedRef<FCollection>& Collection = CachedCollection.Value;
-
-		// Make sure this is a known parent GUID before adding it to the map
-		const FGuid& ParentCollectionGuid = Collection->GetParentCollectionGuid();
-		if (CachedCollectionNamesFromGuids.Contains(ParentCollectionGuid))
-		{
-			auto& CollectionChildren = CachedHierarchy.FindOrAdd(ParentCollectionGuid);
-			CollectionChildren.AddUnique(Collection->GetCollectionGuid());
-		}
-	}
-
-	UE_LOG(LogCollectionManager, Log, TEXT( "Rebuilt the hierarchy cache for %d collections in %0.6f seconds" ), CachedCollections.Num(), FPlatformTime::Seconds() - LoadStartTime);
+	UE_LOG(LogCollectionManager, Log, TEXT( "Loaded %d collections in %0.6f seconds" ), AvailableCollections.Num() - PrevNumCollections, FPlatformTime::Seconds() - LoadStartTime);
 }
 
 bool FCollectionManager::ShouldUseSCC(ECollectionShareType::Type ShareType) const
@@ -1255,14 +1401,14 @@ bool FCollectionManager::AddCollection(const TSharedRef<FCollection>& Collection
 	}
 
 	const FCollectionNameType CollectionKey(CollectionRef->GetCollectionName(), ShareType);
-	if (CachedCollections.Contains(CollectionKey))
+	if (AvailableCollections.Contains(CollectionKey))
 	{
 		UE_LOG(LogCollectionManager, Warning, TEXT("Failed to add collection '%s' because it already exists."), *CollectionRef->GetCollectionName().ToString());
 		return false;
 	}
 
-	CachedCollections.Add(CollectionKey, CollectionRef);
-	CachedCollectionNamesFromGuids.Add(CollectionRef->GetCollectionGuid(), CollectionKey);
+	AvailableCollections.Add(CollectionKey, CollectionRef);
+	CollectionCache.HandleCollectionAdded();
 	return true;
 }
 
@@ -1275,26 +1421,32 @@ bool FCollectionManager::RemoveCollection(const TSharedRef<FCollection>& Collect
 	}
 
 	const FCollectionNameType CollectionKey(CollectionRef->GetCollectionName(), ShareType);
-	CachedCollectionNamesFromGuids.Remove(CollectionRef->GetCollectionGuid());
-	return CachedCollections.Remove(CollectionKey) > 0;
+	if (AvailableCollections.Remove(CollectionKey) > 0)
+	{
+		CollectionCache.HandleCollectionRemoved();
+		return true;
+	}
+
+	return false;
 }
 
 void FCollectionManager::RemoveObjectFromCollections(const FName& ObjectPath, TArray<FCollectionNameType>& OutUpdatedCollections)
 {
-	// Remove the info about the collections that the old object path is within
-	TArray<FObjectCollectionInfo> ObjectCollectionInfos;
-	if (!CachedObjects.RemoveAndCopyValue(ObjectPath, ObjectCollectionInfos))
+	const FCollectionObjectsMap& CachedObjects = CollectionCache.GetCachedObjects();
+
+	const auto* ObjectCollectionInfosPtr = CachedObjects.Find(ObjectPath);
+	if (!ObjectCollectionInfosPtr)
 	{
 		return;
 	}
 
 	// Remove this object reference from all collections that use it
-	for (const FObjectCollectionInfo& ObjectCollectionInfo : ObjectCollectionInfos)
+	for (const FObjectCollectionInfo& ObjectCollectionInfo : *ObjectCollectionInfosPtr)
 	{
 		if ((ObjectCollectionInfo.Reason & ECollectionRecursionFlags::Self) != 0)
 		{
 			// The object is contained directly within this collection (rather than coming from a parent or child collection), so remove the object reference
-			const TSharedRef<FCollection>* const CollectionRefPtr = CachedCollections.Find(ObjectCollectionInfo.CollectionKey);
+			const TSharedRef<FCollection>* const CollectionRefPtr = AvailableCollections.Find(ObjectCollectionInfo.CollectionKey);
 			if (CollectionRefPtr)
 			{
 				OutUpdatedCollections.AddUnique(ObjectCollectionInfo.CollectionKey);
@@ -1307,21 +1459,21 @@ void FCollectionManager::RemoveObjectFromCollections(const FName& ObjectPath, TA
 
 void FCollectionManager::ReplaceObjectInCollections(const FName& OldObjectPath, const FName& NewObjectPath, TArray<FCollectionNameType>& OutUpdatedCollections)
 {
-	// Remove the info about the collections that the old object path is within - we'll add this info back in under the new object path shortly
-	TArray<FObjectCollectionInfo> OldObjectCollectionInfos;
-	if (!CachedObjects.RemoveAndCopyValue(OldObjectPath, OldObjectCollectionInfos))
+	const FCollectionObjectsMap& CachedObjects = CollectionCache.GetCachedObjects();
+
+	const auto* OldObjectCollectionInfosPtr = CachedObjects.Find(OldObjectPath);
+	if (!OldObjectCollectionInfosPtr)
 	{
 		return;
 	}
 
-	// Replace this object reference in all collections that use it, and update our object cache
-	auto& NewObjectCollectionInfos = CachedObjects.FindOrAdd(NewObjectPath);
-	for (const FObjectCollectionInfo& OldObjectCollectionInfo : OldObjectCollectionInfos)
+	// Replace this object reference in all collections that use it
+	for (const FObjectCollectionInfo& OldObjectCollectionInfo : *OldObjectCollectionInfosPtr)
 	{
 		if ((OldObjectCollectionInfo.Reason & ECollectionRecursionFlags::Self) != 0)
 		{
 			// The old object is contained directly within this collection (rather than coming from a parent or child collection), so update the object reference
-			const TSharedRef<FCollection>* const CollectionRefPtr = CachedCollections.Find(OldObjectCollectionInfo.CollectionKey);
+			const TSharedRef<FCollection>* const CollectionRefPtr = AvailableCollections.Find(OldObjectCollectionInfo.CollectionKey);
 			if (CollectionRefPtr)
 			{
 				OutUpdatedCollections.AddUnique(OldObjectCollectionInfo.CollectionKey);
@@ -1330,79 +1482,7 @@ void FCollectionManager::ReplaceObjectInCollections(const FName& OldObjectPath,
 				(*CollectionRefPtr)->AddObjectToCollection(NewObjectPath);
 			}
 		}
-
-		// Merge the collection references for the old object with any collection references that already exist for the new object
-		FObjectCollectionInfo* ObjectInfoPtr = NewObjectCollectionInfos.FindByPredicate([&](const FObjectCollectionInfo& InCollectionInfo) { return InCollectionInfo.CollectionKey == OldObjectCollectionInfo.CollectionKey; });
-		if (ObjectInfoPtr)
-		{
-			ObjectInfoPtr->Reason |= OldObjectCollectionInfo.Reason;
-		}
-		else
-		{
-			NewObjectCollectionInfos.Add(OldObjectCollectionInfo);
-		}
-	}
-}
-
-void FCollectionManager::RecursionHelper_DoWork(const FCollectionNameType& InCollectionKey, const ECollectionRecursionFlags::Flags InRecursionMode, const FRecursiveWorkerFunc& InWorkerFunc) const
-{
-	if ((InRecursionMode & ECollectionRecursionFlags::Self) && InWorkerFunc(InCollectionKey, ECollectionRecursionFlags::Self) == ERecursiveWorkerFlowControl::Stop)
-	{
-		return;
-	}
-
-	if ((InRecursionMode & ECollectionRecursionFlags::Parents) && RecursionHelper_DoWorkOnParents(InCollectionKey, InWorkerFunc) == ERecursiveWorkerFlowControl::Stop)
-	{
-		return;
-	}
-
-	if ((InRecursionMode & ECollectionRecursionFlags::Children) && RecursionHelper_DoWorkOnChildren(InCollectionKey, InWorkerFunc) == ERecursiveWorkerFlowControl::Stop)
-	{
-		return;
-	}
-}
-
-FCollectionManager::ERecursiveWorkerFlowControl FCollectionManager::RecursionHelper_DoWorkOnParents(const FCollectionNameType& InCollectionKey, const FRecursiveWorkerFunc& InWorkerFunc) const
-{
-	const TSharedRef<FCollection>* const CollectionRefPtr = CachedCollections.Find(InCollectionKey);
-	if (CollectionRefPtr)
-	{
-		const FCollectionNameType* const ParentCollectionKeyPtr = CachedCollectionNamesFromGuids.Find((*CollectionRefPtr)->GetParentCollectionGuid());
-		if (ParentCollectionKeyPtr)
-		{
-			if (InWorkerFunc(*ParentCollectionKeyPtr, ECollectionRecursionFlags::Parents) == ERecursiveWorkerFlowControl::Stop || RecursionHelper_DoWorkOnParents(*ParentCollectionKeyPtr, InWorkerFunc) == ERecursiveWorkerFlowControl::Stop)
-			{
-				return ERecursiveWorkerFlowControl::Stop;
-			}
-		}
-	}
-
-	return ERecursiveWorkerFlowControl::Continue;
-}
-
-FCollectionManager::ERecursiveWorkerFlowControl FCollectionManager::RecursionHelper_DoWorkOnChildren(const FCollectionNameType& InCollectionKey, const FRecursiveWorkerFunc& InWorkerFunc) const
-{
-	const TSharedRef<FCollection>* const CollectionRefPtr = CachedCollections.Find(InCollectionKey);
-	if (CollectionRefPtr)
-	{
-		const TArray<FGuid>* const ChildCollectionGuids = CachedHierarchy.Find((*CollectionRefPtr)->GetCollectionGuid());
-		if (ChildCollectionGuids)
-		{
-			for (const FGuid& ChildCollectionGuid : *ChildCollectionGuids)
-			{
-				const FCollectionNameType* const ChildCollectionKeyPtr = CachedCollectionNamesFromGuids.Find(ChildCollectionGuid);
-				if (ChildCollectionKeyPtr)
-				{
-					if (InWorkerFunc(*ChildCollectionKeyPtr, ECollectionRecursionFlags::Children) == ERecursiveWorkerFlowControl::Stop || RecursionHelper_DoWorkOnChildren(*ChildCollectionKeyPtr, InWorkerFunc) == ERecursiveWorkerFlowControl::Stop)
-					{
-						return ERecursiveWorkerFlowControl::Stop;
-					}
-				}
-			}
-		}
 	}
-
-	return ERecursiveWorkerFlowControl::Continue;
 }
 
 #undef LOCTEXT_NAMESPACE
diff --git a/Engine/Source/Developer/CollectionManager/Private/CollectionManager.h b/Engine/Source/Developer/CollectionManager/Private/CollectionManager.h
--- a/Engine/Source/Developer/CollectionManager/Private/CollectionManager.h
+++ b/Engine/Source/Developer/CollectionManager/Private/CollectionManager.h
@@ -2,6 +2,91 @@
 
 #pragma once
 
+/** Collection info for a given object - gives the collection name, as well as the reason this object is considered to be part of this collection */
+struct FObjectCollectionInfo
+{
+	explicit FObjectCollectionInfo(const FCollectionNameType& InCollectionKey)
+		: CollectionKey(InCollectionKey)
+		, Reason(0)
+	{
+	}
+
+	FObjectCollectionInfo(const FCollectionNameType& InCollectionKey, const ECollectionRecursionFlags::Flags InReason)
+		: CollectionKey(InCollectionKey)
+		, Reason(InReason)
+	{
+	}
+
+	/** The key identifying the collection that contains this object */
+	FCollectionNameType CollectionKey;
+	/** The reason(s) why this collection contains this object - this can be tested against the recursion mode when getting the collections for an object */
+	ECollectionRecursionFlags::Flags Reason;
+};
+
+typedef TMap<FCollectionNameType, TSharedRef<FCollection>> FAvailableCollectionsMap;
+typedef TMap<FGuid, FCollectionNameType> FGuidToCollectionNamesMap;
+typedef TMap<FName, TArray<FObjectCollectionInfo>> FCollectionObjectsMap;
+typedef TMap<FGuid, TArray<FGuid>> FCollectionHierarchyMap;
+
+/** Wraps up the lazy caching of the collection manager */
+class FCollectionManagerCache
+{
+public:
+	FCollectionManagerCache(FAvailableCollectionsMap& InAvailableCollections);
+
+	/** Dirty the parts of the cache that need to change when a collection is added to our collection manager */
+	void HandleCollectionAdded();
+	
+	/** Dirty the parts of the cache that need to change when a collection is removed from our collection manager */
+	void HandleCollectionRemoved();
+
+	/** Dirty the parts of the cache that need to change when a collection is modified */
+	void HandleCollectionChanged();
+
+	/** Access the CachedCollectionNamesFromGuids map, ensuring that it is up-to-date */
+	const FGuidToCollectionNamesMap& GetCachedCollectionNamesFromGuids() const;
+
+	/** Access the CachedObjects map, ensuring that it is up-to-date */
+	const FCollectionObjectsMap& GetCachedObjects() const;
+
+	/** Access the CachedHierarchy map, ensuring that it is up-to-date */
+	const FCollectionHierarchyMap& GetCachedHierarchy() const;
+
+	enum class ERecursiveWorkerFlowControl : uint8
+	{
+		Stop,
+		Continue,
+	};
+
+	typedef TFunctionRef<ERecursiveWorkerFlowControl(const FCollectionNameType&, ECollectionRecursionFlags::Flag)> FRecursiveWorkerFunc;
+
+	void RecursionHelper_DoWork(const FCollectionNameType& InCollectionKey, const ECollectionRecursionFlags::Flags InRecursionMode, const FRecursiveWorkerFunc& InWorkerFunc) const;
+	ERecursiveWorkerFlowControl RecursionHelper_DoWorkOnParents(const FCollectionNameType& InCollectionKey, const FRecursiveWorkerFunc& InWorkerFunc) const;
+	ERecursiveWorkerFlowControl RecursionHelper_DoWorkOnChildren(const FCollectionNameType& InCollectionKey, const FRecursiveWorkerFunc& InWorkerFunc) const;
+
+private:
+	/** Reference to the collections that are currently available in our owner collection manager */
+	FAvailableCollectionsMap& AvailableCollections;
+
+	/** A map of collection GUIDs to their associated collection names */
+	mutable FGuidToCollectionNamesMap CachedCollectionNamesFromGuids_Internal;
+
+	/** A map of object paths to their associated collection info - only objects that are in collections will appear in here */
+	mutable FCollectionObjectsMap CachedObjects_Internal;
+
+	/** A map of parent collection GUIDs to their child collection GUIDs - only collections that have children will appear in here */
+	mutable FCollectionHierarchyMap CachedHierarchy_Internal;
+
+	/** Flag to say whether the CachedCollectionNamesFromGuids map is dirty */
+	mutable bool bIsCachedCollectionNamesFromGuidsDirty : 1;
+
+	/** Flag to say whether the CachedObjects map is dirty */
+	mutable bool bIsCachedObjectsDirty : 1;
+
+	/** Flag to say whether the CachedHierarchy map is dirty */
+	mutable bool bIsCachedHierarchyDirty : 1;
+};
+
 class FCollectionManager : public ICollectionManager
 {
 public:
@@ -79,12 +164,6 @@ private:
 	/** Loads all collection files from disk */
 	void LoadCollections();
 
-	/** Rebuild the entire cached objects map based on the current collection data */
-	void RebuildCachedObjects();
-
-	/** Rebuild the entire cached hierarchy map based on the current collection data */
-	void RebuildCachedHierarchy();
-
 	/** Returns true if the specified share type requires source control */
 	bool ShouldUseSCC(ECollectionShareType::Type ShareType) const;
 
@@ -103,40 +182,7 @@ private:
 	/** Replaces an object with another in any collections that contain it */
 	void ReplaceObjectInCollections(const FName& OldObjectPath, const FName& NewObjectPath, TArray<FCollectionNameType>& OutUpdatedCollections);
 
-	enum class ERecursiveWorkerFlowControl : uint8
-	{
-		Stop,
-		Continue,
-	};
-
-	typedef TFunctionRef<ERecursiveWorkerFlowControl(const FCollectionNameType&, ECollectionRecursionFlags::Flag)> FRecursiveWorkerFunc;
-
-	void RecursionHelper_DoWork(const FCollectionNameType& InCollectionKey, const ECollectionRecursionFlags::Flags InRecursionMode, const FRecursiveWorkerFunc& InWorkerFunc) const;
-	ERecursiveWorkerFlowControl RecursionHelper_DoWorkOnParents(const FCollectionNameType& InCollectionKey, const FRecursiveWorkerFunc& InWorkerFunc) const;
-	ERecursiveWorkerFlowControl RecursionHelper_DoWorkOnChildren(const FCollectionNameType& InCollectionKey, const FRecursiveWorkerFunc& InWorkerFunc) const;
-
 private:
-	/** Collection info for a given object - gives the collection name, as well as the reason this object is considered to be part of this collection */
-	struct FObjectCollectionInfo
-	{
-		explicit FObjectCollectionInfo(const FCollectionNameType& InCollectionKey)
-			: CollectionKey(InCollectionKey)
-			, Reason(0)
-		{
-		}
-
-		FObjectCollectionInfo(const FCollectionNameType& InCollectionKey, const ECollectionRecursionFlags::Flags InReason)
-			: CollectionKey(InCollectionKey)
-			, Reason(InReason)
-		{
-		}
-
-		/** The key identifying the collection that contains this object */
-		FCollectionNameType CollectionKey;
-		/** The reason(s) why this collection contains this object - this can be tested against the recursion mode when getting the collections for an object */
-		ECollectionRecursionFlags::Flags Reason;
-	};
-
 	/** The folders that contain collections */
 	FString CollectionFolders[ECollectionShareType::CST_All];
 
@@ -144,16 +190,10 @@ private:
 	FString CollectionExtension;
 
 	/** A map of collection names to FCollection objects */
-	TMap<FCollectionNameType, TSharedRef<FCollection>> CachedCollections;
-
-	/** A map of collection GUIDs to their associated collection names */
-	TMap<FGuid, FCollectionNameType> CachedCollectionNamesFromGuids;
+	TMap<FCollectionNameType, TSharedRef<FCollection>> AvailableCollections;
 
-	/** A map of object paths to their associated collection info - only objects that are in collections will appear in here */
-	TMap<FName, TArray<FObjectCollectionInfo>> CachedObjects;
-
-	/** A map of parent collection GUIDs to their child collection GUIDs - only collections that have children will appear in here */
-	TMap<FGuid, TArray<FGuid>> CachedHierarchy;
+	/** The lazily updated cache for this collection manager */
+	FCollectionManagerCache CollectionCache;
 
 	/** The most recent error that occurred */
 	mutable FText LastError;
