Added file watchers to detect when collections are changed outside of the collection manager

UE-17782 - Editor must be restarted for New Collections to appear after syncing

These use FFileCache to avoid detecting changes that are made by the collection manager itself.

This also fixes some issues when merging collection sets together, and fixes some places where UI wouldn't update when a collection was changed due to an update or merge.

diff --git a/Engine/Source/Developer/CollectionManager/CollectionManager.Build.cs b/Engine/Source/Developer/CollectionManager/CollectionManager.Build.cs
--- a/Engine/Source/Developer/CollectionManager/CollectionManager.Build.cs
+++ b/Engine/Source/Developer/CollectionManager/CollectionManager.Build.cs
@@ -9,6 +9,7 @@ public class CollectionManager : ModuleRules
 		PrivateDependencyModuleNames.AddRange(
 			new string[] {
 				"Core",
+				"DirectoryWatcher",
 				"SourceControl"
 			}
 			);
diff --git a/Engine/Source/Developer/CollectionManager/Private/Collection.cpp b/Engine/Source/Developer/CollectionManager/Private/Collection.cpp
--- a/Engine/Source/Developer/CollectionManager/Private/Collection.cpp
+++ b/Engine/Source/Developer/CollectionManager/Private/Collection.cpp
@@ -312,9 +312,6 @@ bool FCollection::Update(FText& OutError)
 				return false;
 			}
 
-			// Update the info about what the collection looks like on disk
-			DiskSnapshot.TakeSnapshot(NewCollection);
-
 			// Loaded the head revision, now merge up so the files are in a consistent state
 			MergeWithCollection(NewCollection);
 		}
@@ -333,6 +330,11 @@ bool FCollection::Update(FText& OutError)
 	return true;
 }
 
+bool FCollection::Merge(const FCollection& NewCollection)
+{
+	return MergeWithCollection(NewCollection);
+}
+
 bool FCollection::DeleteSourceFile(FText& OutError)
 {
 	bool bSuccessfullyDeleted = false;
@@ -544,58 +546,83 @@ bool FCollection::LoadHeaderPairs(const TMap<FString,FString>& InHeaderPairs)
 	return FileVersion > 0 && FileVersion <= ECollectionVersion::CurrentVersion;
 }
 
-void FCollection::MergeWithCollection(const FCollection& Other)
+bool FCollection::MergeWithCollection(const FCollection& Other)
 {
+	bool bHasChanges = false;
+
 	if ( IsDynamic() )
 	{
 		// @todo collection Merging dynamic collections?
 	}
 	else
 	{
-		// Gather the differences from the file on disk
+		// Work out whether we have any changes compared to the other collection
 		TArray<FName> ObjectsAdded;
 		TArray<FName> ObjectsRemoved;
-		GetObjectDifferencesFromDisk(ObjectsAdded, ObjectsRemoved);
+		GetObjectDifferences(ObjectSet, Other.ObjectSet, ObjectsAdded, ObjectsRemoved);
 
-		// Copy asset list from other collection
-		DiskSnapshot = Other.DiskSnapshot;
-		ObjectSet = DiskSnapshot.ObjectSet;
+		bHasChanges = ParentCollectionGuid != Other.ParentCollectionGuid || ObjectsAdded.Num() > 0 || ObjectsRemoved.Num() > 0;
 
-		// Add the objects that were added before the merge
-		for (const FName& AddedObjectName : ObjectsAdded)
+		if (bHasChanges)
 		{
-			ObjectSet.Add(AddedObjectName);
-		}
+			// Gather the differences from the file on disk
+			ObjectsAdded.Reset();
+			ObjectsRemoved.Reset();
+			GetObjectDifferencesFromDisk(ObjectsAdded, ObjectsRemoved);
+
+			// Copy asset list from other collection
+			DiskSnapshot = Other.DiskSnapshot;
+			ObjectSet = DiskSnapshot.ObjectSet;
+
+			ParentCollectionGuid = Other.ParentCollectionGuid;
+
+			// Add the objects that were added before the merge
+			for (const FName& AddedObjectName : ObjectsAdded)
+			{
+				ObjectSet.Add(AddedObjectName);
+			}
 
-		// Remove the objects that were removed before the merge
-		for (const FName& RemovedObjectName : ObjectsRemoved)
+			// Remove the objects that were removed before the merge
+			for (const FName& RemovedObjectName : ObjectsRemoved)
+			{
+				ObjectSet.Remove(RemovedObjectName);
+			}
+		}
+		else
 		{
-			ObjectSet.Remove(RemovedObjectName);
+			DiskSnapshot = Other.DiskSnapshot;
 		}
 	}
+
+	return bHasChanges;
 }
 
-void FCollection::GetObjectDifferencesFromDisk(TArray<FName>& ObjectsAdded, TArray<FName>& ObjectsRemoved) const
+void FCollection::GetObjectDifferences(const TSet<FName>& BaseSet, const TSet<FName>& NewSet, TArray<FName>& ObjectsAdded, TArray<FName>& ObjectsRemoved)
 {
-	// Find the objects that were removed since the disk list
-	for (const FName& DiskObjectName : DiskSnapshot.ObjectSet)
+	// Find the objects that were removed compared to the base set
+	for (const FName& BaseObjectName : BaseSet)
 	{
-		if (!ObjectSet.Contains(DiskObjectName))
+		if (!NewSet.Contains(BaseObjectName))
 		{
-			ObjectsRemoved.Add(DiskObjectName);
+			ObjectsRemoved.Add(BaseObjectName);
 		}
 	}
 
-	// Find the objects that were added since the disk list
-	for (const FName& MemoryObjectName : ObjectSet)
+	// Find the objects that were added compare to the base set
+	for (const FName& NewObjectName : NewSet)
 	{
-		if (!DiskSnapshot.ObjectSet.Contains(MemoryObjectName))
+		if (!BaseSet.Contains(NewObjectName))
 		{
-			ObjectsAdded.Add(MemoryObjectName);
+			ObjectsAdded.Add(NewObjectName);
 		}
 	}
 }
 
+void FCollection::GetObjectDifferencesFromDisk(TArray<FName>& ObjectsAdded, TArray<FName>& ObjectsRemoved) const
+{
+	GetObjectDifferences(DiskSnapshot.ObjectSet, ObjectSet, ObjectsAdded, ObjectsRemoved);
+}
+
 bool FCollection::CheckoutCollection(FText& OutError)
 {
 	if ( !ensure(SourceFilename.Len()) )
@@ -657,9 +684,6 @@ bool FCollection::CheckoutCollection(FText& OutError)
 				return false;
 			}
 
-			// Update the info about what the collection looks like on disk
-			DiskSnapshot.TakeSnapshot(NewCollection);
-
 			// Loaded the head revision, now merge up so the files are in a consistent state
 			MergeWithCollection(NewCollection);
 		}
@@ -969,9 +993,6 @@ bool FCollection::DeleteFromSourceControl(FText& OutError)
 			return false;
 		}
 
-		// Update the info about what the collection looks like on disk
-		DiskSnapshot.TakeSnapshot(NewCollection);
-
 		// Loaded the head revision, now merge up so the files are in a consistent state
 		MergeWithCollection(NewCollection);
 
diff --git a/Engine/Source/Developer/CollectionManager/Private/Collection.h b/Engine/Source/Developer/CollectionManager/Private/Collection.h
--- a/Engine/Source/Developer/CollectionManager/Private/Collection.h
+++ b/Engine/Source/Developer/CollectionManager/Private/Collection.h
@@ -46,6 +46,8 @@ public:
 	bool Save(FText& OutError);
 	/** Updates this collection to ensure it's the latest version from source control. If false, OutError is a human readable warning depicting the error. */
 	bool Update(FText& OutError);
+	/** Merge the contents of NewCollection into this collection. Returns true if there were changes to merge, or false if the collections were identical. */
+	bool Merge(const FCollection& NewCollection);
 	/** Deletes the source file for this collection. If false, OutError is a human readable warning depicting the error. */
 	bool DeleteSourceFile(FText& OutError);
 
@@ -94,6 +96,9 @@ public:
 	/** Returns the file version of the collection */
 	FORCEINLINE ECollectionVersion::Type GetCollectionVersion() const { return FileVersion; }
 
+	/** Get the source filename of this collection */
+	FORCEINLINE const FString& GetSourceFilename() const { return SourceFilename; }
+
 private:
 	/** Generates the header pairs for the collection file. */
 	void SaveHeaderPairs(TMap<FString,FString>& OutHeaderPairs) const;
@@ -107,8 +112,10 @@ private:
 	bool LoadHeaderPairs(const TMap<FString,FString>& InHeaderPairs);
 
 	/** Merges the assets from the specified collection with this collection */
-	void MergeWithCollection(const FCollection& Other);
-	/** Gets the object differences between collections A and B */
+	bool MergeWithCollection(const FCollection& Other);
+	/** Gets the object differences between object set A (base) and B (new) */
+	static void GetObjectDifferences(const TSet<FName>& BaseSet, const TSet<FName>& NewSet, TArray<FName>& ObjectsAdded, TArray<FName>& ObjectsRemoved);
+	/** Gets the object differences between what we have in memory, and what we loaded from disk */
 	void GetObjectDifferencesFromDisk(TArray<FName>& ObjectsAdded, TArray<FName>& ObjectsRemoved) const;
 	/** Checks the shared collection out from source control so it may be saved. If false, OutError is a human readable warning depicting the error. */
 	bool CheckoutCollection(FText& OutError);
diff --git a/Engine/Source/Developer/CollectionManager/Private/CollectionManager.cpp b/Engine/Source/Developer/CollectionManager/Private/CollectionManager.cpp
--- a/Engine/Source/Developer/CollectionManager/Private/CollectionManager.cpp
+++ b/Engine/Source/Developer/CollectionManager/Private/CollectionManager.cpp
@@ -1,6 +1,8 @@
 // Copyright 1998-2015 Epic Games, Inc. All Rights Reserved.
 
 #include "CollectionManagerPrivatePCH.h"
+#include "FileCache.h"
+#include "Ticker.h"
 
 #define LOCTEXT_NAMESPACE "CollectionManager"
 
@@ -228,10 +230,25 @@ FCollectionManager::FCollectionManager()
 	CollectionExtension = TEXT("collection");
 
 	LoadCollections();
+
+	// Watch for changes that may happen outside of the collection manager
+	for (int32 CacheIdx = 0; CacheIdx < ECollectionShareType::CST_All; ++CacheIdx)
+	{
+		const FString& CollectionFolder = CollectionFolders[CacheIdx];
+
+		DirectoryWatcher::FFileCacheConfig FileCacheConfig(FPaths::ConvertRelativePathToFull(CollectionFolder), FString());
+		FileCacheConfig.DetectMoves(false);
+		FileCacheConfig.RequireFileHashes(false);
+
+		CollectionFileCaches[CacheIdx] = MakeShareable(new DirectoryWatcher::FFileCache(FileCacheConfig));
+	}
+
+	TickFileCacheDelegateHandle = FTicker::GetCoreTicker().AddTicker(FTickerDelegate::CreateRaw(this, &FCollectionManager::TickFileCache), 1.0f);
 }
 
 FCollectionManager::~FCollectionManager()
 {
+	FTicker::GetCoreTicker().RemoveTicker(TickFileCacheDelegateHandle);
 }
 
 bool FCollectionManager::HasCollections() const
@@ -687,6 +704,8 @@ bool FCollectionManager::CreateCollection(FName CollectionName, ECollectionShare
 
 	if (NewCollection->Save(LastError))
 	{
+		CollectionFileCaches[ShareType]->IgnoreNewFile(NewCollection->GetSourceFilename());
+
 		// Collection saved!
 		CollectionCreatedEvent.Broadcast(FCollectionNameType(CollectionName, ShareType));
 		return true;
@@ -746,6 +765,8 @@ bool FCollectionManager::RenameCollection(FName CurrentCollectionName, ECollecti
 	{
 		if ((*CollectionRefPtr)->DeleteSourceFile(LastError))
 		{
+			CollectionFileCaches[CurrentShareType]->IgnoreDeletedFile((*CollectionRefPtr)->GetSourceFilename());
+
 			RemoveCollection(*CollectionRefPtr, CurrentShareType);
 		}
 		else
@@ -757,6 +778,8 @@ bool FCollectionManager::RenameCollection(FName CurrentCollectionName, ECollecti
 		}
 	}
 
+	CollectionFileCaches[NewShareType]->IgnoreNewFile(NewCollection->GetSourceFilename());
+
 	CollectionCache.HandleCollectionChanged();
 
 	// Success
@@ -804,7 +827,11 @@ bool FCollectionManager::ReparentCollection(FName CollectionName, ECollectionSha
 		if ((*ParentCollectionRefPtr)->GetCollectionVersion() < ECollectionVersion::AddedCollectionGuid)
 		{
 			// Try and re-save the parent collection now
-			if (!(*ParentCollectionRefPtr)->Save(LastError))
+			if ((*ParentCollectionRefPtr)->Save(LastError))
+			{
+				CollectionFileCaches[ParentShareType]->IgnoreFileModification((*ParentCollectionRefPtr)->GetSourceFilename());
+			}
+			else
 			{
 				return false;
 			}
@@ -828,7 +855,11 @@ bool FCollectionManager::ReparentCollection(FName CollectionName, ECollectionSha
 	(*CollectionRefPtr)->SetParentCollectionGuid(NewParentGuid);
 
 	// Try and save with the new parent GUID
-	if (!(*CollectionRefPtr)->Save(LastError))
+	if ((*CollectionRefPtr)->Save(LastError))
+	{
+		CollectionFileCaches[ShareType]->IgnoreFileModification((*CollectionRefPtr)->GetSourceFilename());
+	}
+	else
 	{
 		// Failed to save... rollback the collection to use its old parent GUID
 		(*CollectionRefPtr)->SetParentCollectionGuid(OldParentGuid);
@@ -873,6 +904,8 @@ bool FCollectionManager::DestroyCollection(FName CollectionName, ECollectionShar
 
 	if ((*CollectionRefPtr)->DeleteSourceFile(LastError))
 	{
+		CollectionFileCaches[ShareType]->IgnoreDeletedFile((*CollectionRefPtr)->GetSourceFilename());
+
 		RemoveCollection(*CollectionRefPtr, ShareType);
 		CollectionDestroyedEvent.Broadcast(CollectionKey);
 		return true;
@@ -927,6 +960,8 @@ bool FCollectionManager::AddToCollection(FName CollectionName, ECollectionShareT
 	{
 		if ((*CollectionRefPtr)->Save(LastError))
 		{
+			CollectionFileCaches[ShareType]->IgnoreFileModification((*CollectionRefPtr)->GetSourceFilename());
+
 			// Added and saved
 			if (OutNumAdded)
 			{
@@ -1003,6 +1038,8 @@ bool FCollectionManager::RemoveFromCollection(FName CollectionName, ECollectionS
 			
 	if ((*CollectionRefPtr)->Save(LastError))
 	{
+		CollectionFileCaches[ShareType]->IgnoreFileModification((*CollectionRefPtr)->GetSourceFilename());
+
 		// Removed and saved
 		if (OutNumRemoved)
 		{
@@ -1073,6 +1110,8 @@ bool FCollectionManager::SaveCollection(FName CollectionName, ECollectionShareTy
 
 		if ((*CollectionRefPtr)->Save(LastError))
 		{
+			CollectionFileCaches[ShareType]->IgnoreFileModification((*CollectionRefPtr)->GetSourceFilename());
+
 			CollectionCache.HandleCollectionChanged();
 			CollectionUpdatedEvent.Broadcast(CollectionKey);
 			return true;
@@ -1097,6 +1136,8 @@ bool FCollectionManager::UpdateCollection(FName CollectionName, ECollectionShare
 	{
 		if ((*CollectionRefPtr)->Update(LastError))
 		{
+			CollectionFileCaches[ShareType]->IgnoreFileModification((*CollectionRefPtr)->GetSourceFilename());
+
 			CollectionCache.HandleCollectionChanged();
 			CollectionUpdatedEvent.Broadcast(CollectionKey);
 			return true;
@@ -1272,6 +1313,8 @@ bool FCollectionManager::HandleRedirectorDeleted(const FName& ObjectPath)
 			FText SaveError;
 			if (Collection->Save(SaveError))
 			{
+				CollectionFileCaches[CollectionKey.Type]->IgnoreFileModification(Collection->GetSourceFilename());
+
 				UpdatedCollections.Add(CollectionKey);
 			}
 			else
@@ -1343,6 +1386,125 @@ void FCollectionManager::HandleObjectDeleted(const FName& ObjectPath)
 	}
 }
 
+bool FCollectionManager::TickFileCache(float InDeltaTime)
+{
+	enum class ECollectionFileAction : uint8
+	{
+		None,
+		AddCollection,
+		MergeCollection,
+		RemoveCollection,
+	};
+
+	bool bDidChangeCollection = false;
+
+	// Process changes that have happened outside of the collection manager
+	for (int32 CacheIdx = 0; CacheIdx < ECollectionShareType::CST_All; ++CacheIdx)
+	{
+		const ECollectionShareType::Type ShareType = ECollectionShareType::Type(CacheIdx);
+
+		auto& FileCache = CollectionFileCaches[CacheIdx];
+		if (FileCache.IsValid())
+		{
+			FileCache->Tick();
+
+			const auto FileCacheChanges = FileCache->GetOutstandingChanges();
+			for (const auto& FileCacheChange : FileCacheChanges)
+			{
+				const FName CollectionName = *FPaths::GetBaseFilename(FileCacheChange.Filename.Get());
+
+				ECollectionFileAction CollectionFileAction = ECollectionFileAction::None;
+				switch (FileCacheChange.Action)
+				{
+				case DirectoryWatcher::EFileAction::Added:
+				case DirectoryWatcher::EFileAction::Modified:
+					// File was added or modified, but does this collection already exist?
+					CollectionFileAction = (AvailableCollections.Contains(FCollectionNameType(CollectionName, ShareType))) 
+						? ECollectionFileAction::MergeCollection 
+						: ECollectionFileAction::AddCollection;
+					break;
+
+				case DirectoryWatcher::EFileAction::Removed:
+					// File was removed, but does this collection actually exist?
+					CollectionFileAction = (AvailableCollections.Contains(FCollectionNameType(CollectionName, ShareType))) 
+						? ECollectionFileAction::RemoveCollection 
+						: ECollectionFileAction::None;
+					break;
+
+				default:
+					break;
+				}
+
+				switch (CollectionFileAction)
+				{
+				case ECollectionFileAction::AddCollection:
+					{
+						const bool bUseSCC = ShouldUseSCC(ShareType);
+
+						FText LoadErrorText;
+						TSharedRef<FCollection> NewCollection = MakeShareable(new FCollection(GetCollectionFilename(CollectionName, ShareType), bUseSCC));
+						if (NewCollection->Load(LoadErrorText))
+						{
+							if (AddCollection(NewCollection, ShareType))
+							{
+								bDidChangeCollection = true;
+								CollectionCreatedEvent.Broadcast(FCollectionNameType(CollectionName, ShareType));
+							}
+						}
+						else
+						{
+							UE_LOG(LogCollectionManager, Warning, TEXT("%s"), *LoadErrorText.ToString());
+						}
+					}
+					break;
+
+				case ECollectionFileAction::MergeCollection:
+					{
+						TSharedRef<FCollection>* const CollectionRefPtr = AvailableCollections.Find(FCollectionNameType(CollectionName, ShareType));
+						check(CollectionRefPtr); // We tested AvailableCollections.Contains(...) above, so this shouldn't fail
+
+						FText LoadErrorText;
+						FCollection TempCollection(GetCollectionFilename(CollectionName, ShareType), /*bUseSCC*/false);
+						if (TempCollection.Load(LoadErrorText))
+						{
+							if ((*CollectionRefPtr)->Merge(TempCollection))
+							{
+								bDidChangeCollection = true;
+								CollectionUpdatedEvent.Broadcast(FCollectionNameType(CollectionName, ShareType));
+							}
+						}
+						else
+						{
+							UE_LOG(LogCollectionManager, Warning, TEXT("%s"), *LoadErrorText.ToString());
+						}
+					}
+					break;
+
+				case ECollectionFileAction::RemoveCollection:
+					{
+						TSharedRef<FCollection>* const CollectionRefPtr = AvailableCollections.Find(FCollectionNameType(CollectionName, ShareType));
+						check(CollectionRefPtr); // We tested AvailableCollections.Contains(...) above, so this shouldn't fail
+
+						RemoveCollection(*CollectionRefPtr, ShareType);
+						CollectionDestroyedEvent.Broadcast(FCollectionNameType(CollectionName, ShareType));
+					}
+					break;
+
+				default:
+					break;
+				}
+			}
+		}
+	}
+
+	if (bDidChangeCollection)
+	{
+		CollectionCache.HandleCollectionChanged();
+	}
+
+	return true; // Tick again
+}
+
 void FCollectionManager::LoadCollections()
 {
 	const double LoadStartTime = FPlatformTime::Seconds();
@@ -1350,6 +1512,7 @@ void FCollectionManager::LoadCollections()
 
 	for (int32 CacheIdx = 0; CacheIdx < ECollectionShareType::CST_All; ++CacheIdx)
 	{
+		const ECollectionShareType::Type ShareType = ECollectionShareType::Type(CacheIdx);
 		const FString& CollectionFolder = CollectionFolders[CacheIdx];
 		const FString WildCard = FString::Printf(TEXT("%s/*.%s"), *CollectionFolder, *CollectionExtension);
 
@@ -1359,13 +1522,13 @@ void FCollectionManager::LoadCollections()
 		for (const FString& BaseFilename : Filenames)
 		{
 			const FString Filename = CollectionFolder / BaseFilename;
-			const bool bUseSCC = ShouldUseSCC(ECollectionShareType::Type(CacheIdx));
+			const bool bUseSCC = ShouldUseSCC(ShareType);
 
 			FText LoadErrorText;
 			TSharedRef<FCollection> NewCollection = MakeShareable(new FCollection(Filename, bUseSCC));
 			if (NewCollection->Load(LoadErrorText))
 			{
-				AddCollection(NewCollection, ECollectionShareType::Type(CacheIdx));
+				AddCollection(NewCollection, ShareType);
 			}
 			else
 			{
diff --git a/Engine/Source/Developer/CollectionManager/Private/CollectionManager.h b/Engine/Source/Developer/CollectionManager/Private/CollectionManager.h
--- a/Engine/Source/Developer/CollectionManager/Private/CollectionManager.h
+++ b/Engine/Source/Developer/CollectionManager/Private/CollectionManager.h
@@ -2,6 +2,8 @@
 
 #pragma once
 
+namespace DirectoryWatcher { class FFileCache; }
+
 /** Collection info for a given object - gives the collection name, as well as the reason this object is considered to be part of this collection */
 struct FObjectCollectionInfo
 {
@@ -161,6 +163,9 @@ public:
 	virtual FCollectionUpdatedEvent& OnCollectionUpdated() override { return CollectionUpdatedEvent; }
 
 private:
+	/** Tick this collection manager so it can process any file cache events */
+	bool TickFileCache(float InDeltaTime);
+
 	/** Loads all collection files from disk */
 	void LoadCollections();
 
@@ -189,6 +194,12 @@ private:
 	/** The extension used for collection files */
 	FString CollectionExtension;
 
+	/** Array of file cache instances that are watching for the collection files changing on disk */
+	TSharedPtr<DirectoryWatcher::FFileCache> CollectionFileCaches[ECollectionShareType::CST_All];
+
+	/** Delegate handle for the TickFileCache function */
+	FDelegateHandle TickFileCacheDelegateHandle;
+
 	/** A map of collection names to FCollection objects */
 	TMap<FCollectionNameType, TSharedRef<FCollection>> AvailableCollections;
 
diff --git a/Engine/Source/Editor/ContentBrowser/Private/CollectionAssetManagement.cpp b/Engine/Source/Editor/ContentBrowser/Private/CollectionAssetManagement.cpp
--- a/Engine/Source/Editor/ContentBrowser/Private/CollectionAssetManagement.cpp
+++ b/Engine/Source/Editor/ContentBrowser/Private/CollectionAssetManagement.cpp
@@ -15,6 +15,7 @@ FCollectionAssetManagement::FCollectionAssetManagement()
 	// Register the notifications we need in order to keep things up-to-date
 	OnCollectionRenamedHandle = CollectionManagerModule.Get().OnCollectionRenamed().AddRaw(this, &FCollectionAssetManagement::HandleCollectionRenamed);
 	OnCollectionDestroyedHandle = CollectionManagerModule.Get().OnCollectionDestroyed().AddRaw(this, &FCollectionAssetManagement::HandleCollectionDestroyed);
+	OnCollectionUpdatedHandle = CollectionManagerModule.Get().OnCollectionUpdated().AddRaw(this, &FCollectionAssetManagement::HandleCollectionUpdated);
 	OnAssetsAddedHandle = CollectionManagerModule.Get().OnAssetsAdded().AddRaw(this, &FCollectionAssetManagement::HandleAssetsAddedToCollection);
 	OnAssetsRemovedHandle = CollectionManagerModule.Get().OnAssetsRemoved().AddRaw(this, &FCollectionAssetManagement::HandleAssetsRemovedFromCollection);
 }
@@ -28,6 +29,7 @@ FCollectionAssetManagement::~FCollectionAssetManagement()
 
 		CollectionManagerModule.Get().OnCollectionRenamed().Remove(OnCollectionRenamedHandle);
 		CollectionManagerModule.Get().OnCollectionDestroyed().Remove(OnCollectionDestroyedHandle);
+		CollectionManagerModule.Get().OnCollectionUpdated().Remove(OnCollectionUpdatedHandle);
 		CollectionManagerModule.Get().OnAssetsAdded().Remove(OnAssetsAddedHandle);
 		CollectionManagerModule.Get().OnAssetsRemoved().Remove(OnAssetsRemovedHandle);
 	}
@@ -194,6 +196,12 @@ void FCollectionAssetManagement::HandleCollectionRenamed(const FCollectionNameTy
 	}
 }
 
+void FCollectionAssetManagement::HandleCollectionUpdated(const FCollectionNameType& Collection)
+{
+	// Collection has changed in an unknown way - we need to update everything to be sure
+	UpdateAssetManagementState();
+}
+
 void FCollectionAssetManagement::HandleCollectionDestroyed(const FCollectionNameType& Collection)
 {
 	AssetManagementState.Remove(Collection);
diff --git a/Engine/Source/Editor/ContentBrowser/Private/CollectionAssetManagement.h b/Engine/Source/Editor/ContentBrowser/Private/CollectionAssetManagement.h
--- a/Engine/Source/Editor/ContentBrowser/Private/CollectionAssetManagement.h
+++ b/Engine/Source/Editor/ContentBrowser/Private/CollectionAssetManagement.h
@@ -34,6 +34,9 @@ private:
 	/** Handles an on collection renamed event */
 	void HandleCollectionRenamed(const FCollectionNameType& OriginalCollection, const FCollectionNameType& NewCollection);
 
+	/** Handles an on collection updated event */
+	void HandleCollectionUpdated(const FCollectionNameType& Collection);
+
 	/** Handles an on collection destroyed event */
 	void HandleCollectionDestroyed(const FCollectionNameType& Collection);
 
@@ -52,6 +55,7 @@ private:
 	/** Delegate handles */
 	FDelegateHandle OnCollectionRenamedHandle;
 	FDelegateHandle OnCollectionDestroyedHandle;
+	FDelegateHandle OnCollectionUpdatedHandle;
 	FDelegateHandle OnAssetsAddedHandle;
 	FDelegateHandle OnAssetsRemovedHandle;
 };
diff --git a/Engine/Source/Editor/ContentBrowser/Private/SAssetView.cpp b/Engine/Source/Editor/ContentBrowser/Private/SAssetView.cpp
--- a/Engine/Source/Editor/ContentBrowser/Private/SAssetView.cpp
+++ b/Engine/Source/Editor/ContentBrowser/Private/SAssetView.cpp
@@ -92,6 +92,7 @@ void SAssetView::Construct( const FArguments& InArgs )
 	CollectionManagerModule.Get().OnAssetsAdded().AddSP( this, &SAssetView::OnAssetsAddedToCollection );
 	CollectionManagerModule.Get().OnAssetsRemoved().AddSP( this, &SAssetView::OnAssetsRemovedFromCollection );
 	CollectionManagerModule.Get().OnCollectionRenamed().AddSP( this, &SAssetView::OnCollectionRenamed );
+	CollectionManagerModule.Get().OnCollectionUpdated().AddSP( this, &SAssetView::OnCollectionUpdated );
 
 	// Listen for when assets are loaded or changed to update item data
 	FCoreUObjectDelegates::OnAssetLoaded.AddSP(this, &SAssetView::OnAssetLoaded);
@@ -2119,6 +2120,12 @@ void SAssetView::OnCollectionRenamed( const FCollectionNameType& OriginalCollect
 	}
 }
 
+void SAssetView::OnCollectionUpdated( const FCollectionNameType& Collection )
+{
+	// A collection has changed in some way, so we need to refresh our backend list
+	RequestSlowFullListRefresh();
+}
+
 void SAssetView::OnAssetRenamed(const FAssetData& AssetData, const FString& OldObjectPath)
 {
 	// Remove the old asset, if it exists
diff --git a/Engine/Source/Editor/ContentBrowser/Private/SAssetView.h b/Engine/Source/Editor/ContentBrowser/Private/SAssetView.h
--- a/Engine/Source/Editor/ContentBrowser/Private/SAssetView.h
+++ b/Engine/Source/Editor/ContentBrowser/Private/SAssetView.h
@@ -314,6 +314,9 @@ private:
 	/** Handler for when a collection is renamed */
 	void OnCollectionRenamed( const FCollectionNameType& OriginalCollection, const FCollectionNameType& NewCollection );
 
+	/** Handler for when a collection is updated */
+	void OnCollectionUpdated( const FCollectionNameType& Collection );
+
 	/** Handler for when an asset was renamed in the asset registry */
 	void OnAssetRenamed(const FAssetData& AssetData, const FString& OldObjectPath);
 
