Added support for dynamic collections

UETOOL-332 - Collections 2.0
UETOOL-401 - Add support for dynamic collections

These collections can be created from a text search in the Content Browser, and allow you to store a search query so it can be reused or shared.

They have some limitations compared to static collections:
- As they never know what objects they actually contain, they can't be displayed in the "Collections" list in asset tooltips.
- They cannot contain child collections.

diff --git a/Engine/Source/Developer/CollectionManager/Private/Collection.cpp b/Engine/Source/Developer/CollectionManager/Private/Collection.cpp
--- a/Engine/Source/Developer/CollectionManager/Private/Collection.cpp
+++ b/Engine/Source/Developer/CollectionManager/Private/Collection.cpp
@@ -17,7 +17,7 @@ struct FCollectionUtils
 	}
 };
 
-FCollection::FCollection(const FString& InFilename, bool InUseSCC)
+FCollection::FCollection(const FString& InFilename, bool InUseSCC, ECollectionStorageMode::Type InStorageMode)
 {
 	ensure(InFilename.Len() > 0);
 
@@ -25,6 +25,8 @@ FCollection::FCollection(const FString& InFilename, bool InUseSCC)
 	SourceFilename = InFilename;
 	CollectionName = FName(*FPaths::GetBaseFilename(InFilename));
 
+	StorageMode = InStorageMode;
+
 	CollectionGuid = FGuid::NewGuid();
 
 	// Initialize the file version to the most recent
@@ -40,6 +42,8 @@ TSharedRef<FCollection> FCollection::Clone(const FString& InFilename, bool InUse
 	NewCollection->SourceFilename = InFilename;
 	NewCollection->CollectionName = FName(*FPaths::GetBaseFilename(InFilename));
 
+	NewCollection->StorageMode = StorageMode;
+
 	// Create a new GUID?
 	if (InCloneMode == ECollectionCloneMode::Unique)
 	{
@@ -103,11 +107,7 @@ bool FCollection::Load(FText& OutError)
 	}
 
 	// Now load the content if the header load was successful
-	if ( IsDynamic() )
-	{
-		// @todo collection Load dynamic collections
-	}
-	else
+	if (StorageMode == ECollectionStorageMode::Static)
 	{
 		// Static collection, a flat list of asset paths
 		for (FString Line : FileContents)
@@ -121,6 +121,14 @@ bool FCollection::Load(FText& OutError)
 			}
 		}
 	}
+	else
+	{
+		// Dynamic collection, a single query line
+		DynamicQueryText = (FileContents.Num() > 0) ? FileContents[0] : FString();
+
+		DynamicQueryText.Trim();
+		DynamicQueryText.TrimTrailing();
+	}
 
 	DiskSnapshot.TakeSnapshot(*this);
 
@@ -172,11 +180,7 @@ bool FCollection::Save(FText& OutError)
 	FileOutput += LINE_TERMINATOR;
 
 	// Now for the content
-	if ( IsDynamic() )
-	{
-		// @todo Dynamic collections
-	}
-	else
+	if (StorageMode == ECollectionStorageMode::Static)
 	{
 		// Write out the set as a sorted array to keep things in a known order for diffing
 		TArray<FName> ObjectList = ObjectSet.Array();
@@ -188,6 +192,11 @@ bool FCollection::Save(FText& OutError)
 			FileOutput += ObjectName.ToString() + LINE_TERMINATOR;
 		}
 	}
+	else
+	{
+		// Dynamic collection, a single query line
+		FileOutput += DynamicQueryText + LINE_TERMINATOR;
+	}
 
 	// Attempt to save the file
 	bool bSaveSuccessful = false;
@@ -304,7 +313,7 @@ bool FCollection::Update(FText& OutError)
 		{
 			// File found! Load it and merge with our local changes
 			FText LoadErrorText;
-			FCollection NewCollection(SourceFilename, false);
+			FCollection NewCollection(SourceFilename, false, ECollectionStorageMode::Static);
 			if ( !NewCollection.Load(LoadErrorText) )
 			{
 				// Failed to load the head revision file so it isn't safe to delete it
@@ -370,7 +379,7 @@ bool FCollection::DeleteSourceFile(FText& OutError)
 
 bool FCollection::AddObjectToCollection(FName ObjectPath)
 {
-	if (!ObjectSet.Contains(ObjectPath))
+	if (StorageMode == ECollectionStorageMode::Static && !ObjectSet.Contains(ObjectPath))
 	{
 		ObjectSet.Add(ObjectPath);
 		return true;
@@ -381,46 +390,86 @@ bool FCollection::AddObjectToCollection(FName ObjectPath)
 
 bool FCollection::RemoveObjectFromCollection(FName ObjectPath)
 {
-	return ObjectSet.Remove(ObjectPath) > 0;
+	if (StorageMode == ECollectionStorageMode::Static)
+	{
+		return ObjectSet.Remove(ObjectPath) > 0;
+	}
+
+	return false;
 }
 
 void FCollection::GetAssetsInCollection(TArray<FName>& Assets) const
 {
-	for (const FName& ObjectName : ObjectSet)
+	if (StorageMode == ECollectionStorageMode::Static)
 	{
-		if (!ObjectName.ToString().StartsWith(TEXT("/Script/")))
+		for (const FName& ObjectName : ObjectSet)
 		{
-			Assets.Add(ObjectName);
+			if (!ObjectName.ToString().StartsWith(TEXT("/Script/")))
+			{
+				Assets.Add(ObjectName);
+			}
 		}
 	}
 }
 
 void FCollection::GetClassesInCollection(TArray<FName>& Classes) const
 {
-	for (const FName& ObjectName : ObjectSet)
+	if (StorageMode == ECollectionStorageMode::Static)
 	{
-		if (ObjectName.ToString().StartsWith(TEXT("/Script/")))
+		for (const FName& ObjectName : ObjectSet)
 		{
-			Classes.Add(ObjectName);
+			if (ObjectName.ToString().StartsWith(TEXT("/Script/")))
+			{
+				Classes.Add(ObjectName);
+			}
 		}
 	}
 }
 
 void FCollection::GetObjectsInCollection(TArray<FName>& Objects) const
 {
-	FCollectionUtils::AppendCollectionToArray(ObjectSet, Objects);
+	if (StorageMode == ECollectionStorageMode::Static)
+	{
+		FCollectionUtils::AppendCollectionToArray(ObjectSet, Objects);
+	}
 }
 
 bool FCollection::IsObjectInCollection(FName ObjectPath) const
 {
-	return ObjectSet.Contains(ObjectPath);
+	if (StorageMode == ECollectionStorageMode::Static)
+	{
+		return ObjectSet.Contains(ObjectPath);
+	}
+
+	return false;
 }
 
 bool FCollection::IsRedirectorInCollection(FName ObjectPath) const
 {
-	// Redirectors are fixed up in-memory once the asset registry has finished loading, 
-	// so we need to test our on-disk set of objects rather than our in-memory set of objects
-	return DiskSnapshot.ObjectSet.Contains(ObjectPath);
+	if (StorageMode == ECollectionStorageMode::Static)
+	{
+		// Redirectors are fixed up in-memory once the asset registry has finished loading, 
+		// so we need to test our on-disk set of objects rather than our in-memory set of objects
+		return DiskSnapshot.ObjectSet.Contains(ObjectPath);
+	}
+
+	return false;
+}
+
+bool FCollection::SetDynamicQueryText(const FString& InQueryText)
+{
+	if (StorageMode == ECollectionStorageMode::Dynamic)
+	{
+		DynamicQueryText = InQueryText;
+		return true;
+	}
+
+	return false;
+}
+
+FString FCollection::GetDynamicQueryText() const
+{
+	return (StorageMode == ECollectionStorageMode::Dynamic) ? DynamicQueryText : FString();
 }
 
 FCollectionStatusInfo FCollection::GetStatusInfo() const
@@ -428,7 +477,7 @@ FCollectionStatusInfo FCollection::GetStatusInfo() const
 	FCollectionStatusInfo StatusInfo;
 
 	StatusInfo.bIsDirty = IsDirty();
-	StatusInfo.bIsEmpty = ObjectSet.Num() == 0;
+	StatusInfo.bIsEmpty = IsEmpty();
 	StatusInfo.bUseSCC  = bUseSCC;
 
 	StatusInfo.NumObjects = ObjectSet.Num();
@@ -453,31 +502,39 @@ bool FCollection::IsDirty() const
 		return true;
 	}
 	
-	TArray<FName> ObjectsAdded;
-	TArray<FName> ObjectsRemoved;
-	GetObjectDifferencesFromDisk(ObjectsAdded, ObjectsRemoved);
+	bool bHasChanges = false;
 
-	return ObjectsAdded.Num() != 0 || ObjectsRemoved.Num() != 0;
-}
+	if (StorageMode == ECollectionStorageMode::Static)
+	{
+		TArray<FName> ObjectsAdded;
+		TArray<FName> ObjectsRemoved;
+		GetObjectDifferencesFromDisk(ObjectsAdded, ObjectsRemoved);
 
-bool FCollection::IsDynamic() const
-{
-	// @todo collection Dynamic collections
-	return false;
+		bHasChanges = ObjectsAdded.Num() != 0 || ObjectsRemoved.Num() != 0;
+	}
+	else
+	{
+		bHasChanges = DynamicQueryText != DiskSnapshot.DynamicQueryText;
+	}
+
+	return bHasChanges;
 }
 
 bool FCollection::IsEmpty() const
 {
-	return ObjectSet.Num() == 0;
+	if (StorageMode == ECollectionStorageMode::Static)
+	{
+		return ObjectSet.Num() == 0;
+	}
+	else
+	{
+		return DynamicQueryText.IsEmpty();
+	}
 }
 
 void FCollection::PrintCollection() const
 {
-	if ( IsDynamic() )
-	{
-		// @todo collection Printing Dynamic collections
-	}
-	else
+	if (StorageMode == ECollectionStorageMode::Static)
 	{
 		UE_LOG(LogCollectionManager, Log, TEXT("    Printing static elements of collection %s"), *CollectionName.ToString());
 		UE_LOG(LogCollectionManager, Log, TEXT("    ============================="));
@@ -491,13 +548,19 @@ void FCollection::PrintCollection() const
 			UE_LOG(LogCollectionManager, Log, TEXT("        %s"), *ObjectName.ToString());
 		}
 	}
+	else
+	{
+		UE_LOG(LogCollectionManager, Log, TEXT("    Printing dynamic query of collection %s"), *CollectionName.ToString());
+		UE_LOG(LogCollectionManager, Log, TEXT("    ============================="));
+		UE_LOG(LogCollectionManager, Log, TEXT("        %s"), *DynamicQueryText);
+	}
 }
 
 void FCollection::SaveHeaderPairs(TMap<FString,FString>& OutHeaderPairs) const
 {
 	// These pairs will appear at the top of the file followed by a newline
 	OutHeaderPairs.Add(TEXT("FileVersion"), FString::FromInt(ECollectionVersion::CurrentVersion)); // Files are always saved at the latest version as loading should take care of data upgrades
-	OutHeaderPairs.Add(TEXT("Type"), IsDynamic() ? TEXT("Dynamic") : TEXT("Static"));
+	OutHeaderPairs.Add(TEXT("Type"), ECollectionStorageMode::ToString(StorageMode));
 	OutHeaderPairs.Add(TEXT("Guid"), CollectionGuid.ToString(EGuidFormats::DigitsWithHyphens));
 	OutHeaderPairs.Add(TEXT("ParentGuid"), ParentCollectionGuid.ToString(EGuidFormats::DigitsWithHyphens));
 }
@@ -520,6 +583,8 @@ bool FCollection::LoadHeaderPairs(const TMap<FString,FString>& InHeaderPairs)
 		return false;
 	}
 
+	StorageMode = ECollectionStorageMode::FromString(**Type);
+
 	FileVersion = (ECollectionVersion::Type)FCString::Atoi(**Version);
 
 	if (FileVersion >= ECollectionVersion::AddedCollectionGuid)
@@ -538,30 +603,35 @@ bool FCollection::LoadHeaderPairs(const TMap<FString,FString>& InHeaderPairs)
 		}
 	}
 
-	if ( *Type == TEXT("Dynamic") )
-	{
-		// @todo Set this file up to be dynamic
-	}
-
 	return FileVersion > 0 && FileVersion <= ECollectionVersion::CurrentVersion;
 }
 
 bool FCollection::MergeWithCollection(const FCollection& Other)
 {
-	bool bHasChanges = false;
+	bool bHasChanges = ParentCollectionGuid != Other.ParentCollectionGuid;
+
+	ParentCollectionGuid = Other.ParentCollectionGuid;
 
-	if ( IsDynamic() )
+	if (StorageMode != Other.StorageMode)
 	{
-		// @todo collection Merging dynamic collections?
+		bHasChanges = true;
+		StorageMode = Other.StorageMode;
+
+		// Storage mode has changed!
+		// Reset the static and dynamic parts of the collection so we just copy over the new data verbatim
+		ObjectSet.Reset();
+		DynamicQueryText.Reset();
+		DiskSnapshot.TakeSnapshot(*this);
 	}
-	else
+
+	if (StorageMode == ECollectionStorageMode::Static)
 	{
 		// Work out whether we have any changes compared to the other collection
 		TArray<FName> ObjectsAdded;
 		TArray<FName> ObjectsRemoved;
 		GetObjectDifferences(ObjectSet, Other.ObjectSet, ObjectsAdded, ObjectsRemoved);
 
-		bHasChanges = ParentCollectionGuid != Other.ParentCollectionGuid || ObjectsAdded.Num() > 0 || ObjectsRemoved.Num() > 0;
+		bHasChanges = bHasChanges || ObjectsAdded.Num() > 0 || ObjectsRemoved.Num() > 0;
 
 		if (bHasChanges)
 		{
@@ -571,10 +641,7 @@ bool FCollection::MergeWithCollection(const FCollection& Other)
 			GetObjectDifferencesFromDisk(ObjectsAdded, ObjectsRemoved);
 
 			// Copy asset list from other collection
-			DiskSnapshot = Other.DiskSnapshot;
-			ObjectSet = DiskSnapshot.ObjectSet;
-
-			ParentCollectionGuid = Other.ParentCollectionGuid;
+			ObjectSet = Other.ObjectSet;
 
 			// Add the objects that were added before the merge
 			for (const FName& AddedObjectName : ObjectsAdded)
@@ -588,11 +655,14 @@ bool FCollection::MergeWithCollection(const FCollection& Other)
 				ObjectSet.Remove(RemovedObjectName);
 			}
 		}
-		else
-		{
-			DiskSnapshot = Other.DiskSnapshot;
-		}
 	}
+	else
+	{
+		bHasChanges = bHasChanges || DynamicQueryText != Other.DynamicQueryText;
+		DynamicQueryText = Other.DynamicQueryText;
+	}
+
+	DiskSnapshot = Other.DiskSnapshot;
 
 	return bHasChanges;
 }
@@ -620,7 +690,10 @@ void FCollection::GetObjectDifferences(const TSet<FName>& BaseSet, const TSet<FN
 
 void FCollection::GetObjectDifferencesFromDisk(TArray<FName>& ObjectsAdded, TArray<FName>& ObjectsRemoved) const
 {
-	GetObjectDifferences(DiskSnapshot.ObjectSet, ObjectSet, ObjectsAdded, ObjectsRemoved);
+	if (StorageMode == ECollectionStorageMode::Static)
+	{
+		GetObjectDifferences(DiskSnapshot.ObjectSet, ObjectSet, ObjectsAdded, ObjectsRemoved);
+	}
 }
 
 bool FCollection::CheckoutCollection(FText& OutError)
@@ -676,7 +749,7 @@ bool FCollection::CheckoutCollection(FText& OutError)
 		{
 			// File found! Load it and merge with our local changes
 			FText LoadErrorText;
-			FCollection NewCollection(SourceFilename, false);
+			FCollection NewCollection(SourceFilename, false, ECollectionStorageMode::Static);
 			if ( !NewCollection.Load(LoadErrorText) )
 			{
 				// Failed to load the head revision file so it isn't safe to delete it
@@ -786,11 +859,7 @@ bool FCollection::CheckinCollection(FText& OutError)
 	}
 	else
 	{
-		if (IsDynamic())
-		{
-			// @todo collection Change description for dynamic collections
-		}
-		else
+		if (StorageMode == ECollectionStorageMode::Static)
 		{
 			// Gather differences from disk
 			TArray<FName> ObjectsAdded;
@@ -840,6 +909,13 @@ bool FCollection::CheckinCollection(FText& OutError)
 				ChangelistDescBuilder.Unindent();
 			}
 		}
+		else
+		{
+			if (DiskSnapshot.DynamicQueryText != DynamicQueryText)
+			{
+				ChangelistDescBuilder.AppendLineFormat(LOCTEXT("CollectionChangedParentDesc", "Changed the dynamic query of collection '{0}' to '{1}'"), CollectionNameText, FText::FromString(DynamicQueryText));
+			}
+		}
 
 		// Parent change?
 		if (DiskSnapshot.ParentCollectionGuid != ParentCollectionGuid)
@@ -983,7 +1059,7 @@ bool FCollection::DeleteFromSourceControl(FText& OutError)
 			return true;
 		}
 			
-		FCollection NewCollection(SourceFilename, false);
+		FCollection NewCollection(SourceFilename, false, ECollectionStorageMode::Static);
 		FText LoadErrorText;
 		if ( !NewCollection.Load(LoadErrorText) )
 		{
diff --git a/Engine/Source/Developer/CollectionManager/Private/Collection.h b/Engine/Source/Developer/CollectionManager/Private/Collection.h
--- a/Engine/Source/Developer/CollectionManager/Private/Collection.h
+++ b/Engine/Source/Developer/CollectionManager/Private/Collection.h
@@ -35,7 +35,7 @@ enum class ECollectionCloneMode : uint8
 class FCollection
 {
 public:
-	FCollection(const FString& InFilename, bool InUseSCC);
+	FCollection(const FString& InFilename, bool InUseSCC, ECollectionStorageMode::Type InStorageMode);
 
 	/** Clone this collection to a new location */
 	TSharedRef<FCollection> Clone(const FString& InFilename, bool InUseSCC, ECollectionCloneMode InCloneMode) const;
@@ -66,15 +66,17 @@ public:
 	/** Returns true when the specified redirector is in the collection. Static collections only. */
 	bool IsRedirectorInCollection(FName ObjectPath) const;
 
+	/** Set the dynamic query text for this collection. Dynamic collections only. */
+	bool SetDynamicQueryText(const FString& InQueryText);
+	/** Get the dynamic query text for this collection. Dynamic collections only. */
+	FString GetDynamicQueryText() const;
+
 	/** Get the status info for this collection */
 	FCollectionStatusInfo GetStatusInfo() const;
 
 	/** Does this collection contain unsaved changes? */
 	bool IsDirty() const;
 
-	/** Returns true if the collection contains rules instead of a flat list */
-	bool IsDynamic() const;
-
 	/** Whether the collection has any contents */
 	bool IsEmpty() const;
 
@@ -96,6 +98,9 @@ public:
 	/** Returns the file version of the collection */
 	FORCEINLINE ECollectionVersion::Type GetCollectionVersion() const { return FileVersion; }
 
+	/** Get whether this collection is static or dynamic */
+	ECollectionStorageMode::Type GetStorageMode() const { return StorageMode; }
+
 	/** Get the source filename of this collection */
 	FORCEINLINE const FString& GetSourceFilename() const { return SourceFilename; }
 
@@ -115,7 +120,7 @@ private:
 	bool MergeWithCollection(const FCollection& Other);
 	/** Gets the object differences between object set A (base) and B (new) */
 	static void GetObjectDifferences(const TSet<FName>& BaseSet, const TSet<FName>& NewSet, TArray<FName>& ObjectsAdded, TArray<FName>& ObjectsRemoved);
-	/** Gets the object differences between what we have in memory, and what we loaded from disk */
+	/** Gets the object differences between what we have in memory, and what we loaded from disk. Static collections only. */
 	void GetObjectDifferencesFromDisk(TArray<FName>& ObjectsAdded, TArray<FName>& ObjectsRemoved) const;
 	/** Checks the shared collection out from source control so it may be saved. If false, OutError is a human readable warning depicting the error. */
 	bool CheckoutCollection(FText& OutError);
@@ -134,13 +139,17 @@ private:
 		{
 			ParentCollectionGuid = InCollection.ParentCollectionGuid;
 			ObjectSet = InCollection.ObjectSet;
+			DynamicQueryText = InCollection.DynamicQueryText;
 		}
 
 		/** The GUID of the collection we are parented under */
 		FGuid ParentCollectionGuid;
 
-		/** The set of objects in the collection. Takes the form PackageName.AssetName */
+		/** The set of objects in the collection. Takes the form PackageName.AssetName. Static collections only. */
 		TSet<FName> ObjectSet;
+
+		/** The dynamic query string for this collection. Dynamic collections only. */
+		FString DynamicQueryText;
 	};
 
 	/** The name of the collection */
@@ -158,12 +167,18 @@ private:
 	/** The filename used to load this collection. Empty if it is new or never loaded from disk. */
 	FString SourceFilename;
 
-	/** The set of objects in the collection. Takes the form PackageName.AssetName */
+	/** The set of objects in the collection. Takes the form PackageName.AssetName.Static collections only. */
 	TSet<FName> ObjectSet;
 
+	/** The dynamic query string for this collection. Dynamic collections only. */
+	FString DynamicQueryText;
+
 	/** The file version for this collection */
 	ECollectionVersion::Type FileVersion;
 
+	/** How does this collection store its objects? (static or dynamic) */
+	ECollectionStorageMode::Type StorageMode;
+
 	/** The state of the collection the last time it was loaded from or saved to disk. */
 	FCollectionSnapshot DiskSnapshot;
 };
diff --git a/Engine/Source/Developer/CollectionManager/Private/CollectionManager.cpp b/Engine/Source/Developer/CollectionManager/Private/CollectionManager.cpp
--- a/Engine/Source/Developer/CollectionManager/Private/CollectionManager.cpp
+++ b/Engine/Source/Developer/CollectionManager/Private/CollectionManager.cpp
@@ -681,7 +681,7 @@ bool FCollectionManager::IsValidCollectionName(const FString& CollectionName, EC
 	return true;
 }
 
-bool FCollectionManager::CreateCollection(FName CollectionName, ECollectionShareType::Type ShareType)
+bool FCollectionManager::CreateCollection(FName CollectionName, ECollectionShareType::Type ShareType, ECollectionStorageMode::Type StorageMode)
 {
 	if (!ensure(ShareType < ECollectionShareType::CST_All))
 	{
@@ -694,7 +694,7 @@ bool FCollectionManager::CreateCollection(FName CollectionName, ECollectionShare
 	const bool bUseSCC = ShouldUseSCC(ShareType);
 	const FString CollectionFilename = GetCollectionFilename(CollectionName, ShareType);
 
-	TSharedRef<FCollection> NewCollection = MakeShareable(new FCollection(CollectionFilename, bUseSCC));
+	TSharedRef<FCollection> NewCollection = MakeShareable(new FCollection(CollectionFilename, bUseSCC, StorageMode));
 	if (!AddCollection(NewCollection, ShareType))
 	{
 		// Failed to add the collection, it already exists
@@ -947,6 +947,12 @@ bool FCollectionManager::AddToCollection(FName CollectionName, ECollectionShareT
 		return false;
 	}
 
+	if ((*CollectionRefPtr)->GetStorageMode() != ECollectionStorageMode::Static)
+	{
+		LastError = LOCTEXT("Error_AddNeedsStaticCollection", "Objects can only be added to static collections.");
+		return false;
+	}
+
 	int32 NumAdded = 0;
 	for (const FName& ObjectPath : ObjectPaths)
 	{
@@ -1020,6 +1026,12 @@ bool FCollectionManager::RemoveFromCollection(FName CollectionName, ECollectionS
 		return false;
 	}
 
+	if ((*CollectionRefPtr)->GetStorageMode() != ECollectionStorageMode::Static)
+	{
+		LastError = LOCTEXT("Error_RemoveNeedsStaticCollection", "Objects can only be removed from static collections.");
+		return false;
+	}
+
 	TArray<FName> RemovedAssets;
 	for (const FName& ObjectPath : ObjectPaths)
 	{
@@ -1061,6 +1073,72 @@ bool FCollectionManager::RemoveFromCollection(FName CollectionName, ECollectionS
 	}
 }
 
+bool FCollectionManager::SetDynamicQueryText(FName CollectionName, ECollectionShareType::Type ShareType, const FString& InQueryText)
+{
+	if (!ensure(ShareType < ECollectionShareType::CST_All))
+	{
+		// Bad share type
+		LastError = LOCTEXT("Error_Internal", "There was an internal error.");
+		return false;
+	}
+
+	const FCollectionNameType CollectionKey(CollectionName, ShareType);
+	TSharedRef<FCollection>* const CollectionRefPtr = AvailableCollections.Find(CollectionKey);
+	if (!CollectionRefPtr)
+	{
+		// Collection doesn't exist
+		LastError = LOCTEXT("Error_DoesntExist", "The collection doesn't exist.");
+		return false;
+	}
+
+	if ((*CollectionRefPtr)->GetStorageMode() != ECollectionStorageMode::Dynamic)
+	{
+		LastError = LOCTEXT("Error_SetNeedsDynamicCollection", "Search queries can only be set on dynamic collections.");
+		return false;
+	}
+
+	(*CollectionRefPtr)->SetDynamicQueryText(InQueryText);
+	
+	if ((*CollectionRefPtr)->Save(LastError))
+	{
+		CollectionFileCaches[ShareType]->IgnoreFileModification((*CollectionRefPtr)->GetSourceFilename());
+
+		CollectionCache.HandleCollectionChanged();
+		CollectionUpdatedEvent.Broadcast(CollectionKey);
+		return true;
+	}
+
+	return false;
+}
+
+bool FCollectionManager::GetDynamicQueryText(FName CollectionName, ECollectionShareType::Type ShareType, FString& OutQueryText) const
+{
+	if (!ensure(ShareType < ECollectionShareType::CST_All))
+	{
+		// Bad share type
+		LastError = LOCTEXT("Error_Internal", "There was an internal error.");
+		return false;
+	}
+
+	const FCollectionNameType CollectionKey(CollectionName, ShareType);
+	const TSharedRef<FCollection>* const CollectionRefPtr = AvailableCollections.Find(CollectionKey);
+	if (!CollectionRefPtr)
+	{
+		// Collection doesn't exist
+		LastError = LOCTEXT("Error_DoesntExist", "The collection doesn't exist.");
+		return false;
+	}
+
+	if ((*CollectionRefPtr)->GetStorageMode() != ECollectionStorageMode::Dynamic)
+	{
+		LastError = LOCTEXT("Error_GetNeedsDynamicCollection", "Search queries can only be got from dynamic collections.");
+		return false;
+	}
+
+	OutQueryText = (*CollectionRefPtr)->GetDynamicQueryText();
+	return true;
+}
+
 bool FCollectionManager::EmptyCollection(FName CollectionName, ECollectionShareType::Type ShareType)
 {
 	if (!ensure(ShareType < ECollectionShareType::CST_All))
@@ -1117,6 +1195,10 @@ bool FCollectionManager::SaveCollection(FName CollectionName, ECollectionShareTy
 			return true;
 		}
 	}
+	else
+	{
+		LastError = LOCTEXT("Error_DoesntExist", "The collection doesn't exist.");
+	}
 
 	return false;
 }
@@ -1143,6 +1225,10 @@ bool FCollectionManager::UpdateCollection(FName CollectionName, ECollectionShare
 			return true;
 		}
 	}
+	else
+	{
+		LastError = LOCTEXT("Error_DoesntExist", "The collection doesn't exist.");
+	}
 
 	return false;
 }
@@ -1163,8 +1249,36 @@ bool FCollectionManager::GetCollectionStatusInfo(FName CollectionName, ECollecti
 		OutStatusInfo = (*CollectionRefPtr)->GetStatusInfo();
 		return true;
 	}
+	else
+	{
+		LastError = LOCTEXT("Error_DoesntExist", "The collection doesn't exist.");
+	}
 
-	return true;
+	return false;
+}
+
+bool FCollectionManager::GetCollectionStorageMode(FName CollectionName, ECollectionShareType::Type ShareType, ECollectionStorageMode::Type& OutStorageMode) const
+{
+	if (!ensure(ShareType < ECollectionShareType::CST_All))
+	{
+		// Bad share type
+		LastError = LOCTEXT("Error_Internal", "There was an internal error.");
+		return true;
+	}
+
+	const FCollectionNameType CollectionKey(CollectionName, ShareType);
+	const TSharedRef<FCollection>* const CollectionRefPtr = AvailableCollections.Find(CollectionKey);
+	if (CollectionRefPtr)
+	{
+		OutStorageMode = (*CollectionRefPtr)->GetStorageMode();
+		return true;
+	}
+	else
+	{
+		LastError = LOCTEXT("Error_DoesntExist", "The collection doesn't exist.");
+	}
+
+	return false;
 }
 
 bool FCollectionManager::IsObjectInCollection(FName ObjectPath, FName CollectionName, ECollectionShareType::Type ShareType, ECollectionRecursionFlags::Flags RecursionMode) const
@@ -1230,6 +1344,18 @@ bool FCollectionManager::IsValidParentCollection(FName CollectionName, ECollecti
 			return FCollectionManagerCache::ERecursiveWorkerFlowControl::Stop;
 		}
 
+		const TSharedRef<FCollection>* const CollectionRefPtr = AvailableCollections.Find(InCollectionKey);
+		if (CollectionRefPtr)
+		{
+			const ECollectionStorageMode::Type StorageMode = (*CollectionRefPtr)->GetStorageMode();
+			if (StorageMode == ECollectionStorageMode::Dynamic)
+			{
+				bValidParent = false;
+				LastError = LOCTEXT("InvalidParent_InvalidParentStorageType", "A dynamic collection cannot contain child collections");
+				return FCollectionManagerCache::ERecursiveWorkerFlowControl::Stop;
+			}
+		}
+
 		return FCollectionManagerCache::ERecursiveWorkerFlowControl::Continue;
 	};
 
@@ -1442,7 +1568,7 @@ bool FCollectionManager::TickFileCache(float InDeltaTime)
 						const bool bUseSCC = ShouldUseSCC(ShareType);
 
 						FText LoadErrorText;
-						TSharedRef<FCollection> NewCollection = MakeShareable(new FCollection(GetCollectionFilename(CollectionName, ShareType), bUseSCC));
+						TSharedRef<FCollection> NewCollection = MakeShareable(new FCollection(GetCollectionFilename(CollectionName, ShareType), bUseSCC, ECollectionStorageMode::Static));
 						if (NewCollection->Load(LoadErrorText))
 						{
 							if (AddCollection(NewCollection, ShareType))
@@ -1464,7 +1590,7 @@ bool FCollectionManager::TickFileCache(float InDeltaTime)
 						check(CollectionRefPtr); // We tested AvailableCollections.Contains(...) above, so this shouldn't fail
 
 						FText LoadErrorText;
-						FCollection TempCollection(GetCollectionFilename(CollectionName, ShareType), /*bUseSCC*/false);
+						FCollection TempCollection(GetCollectionFilename(CollectionName, ShareType), /*bUseSCC*/false, ECollectionStorageMode::Static);
 						if (TempCollection.Load(LoadErrorText))
 						{
 							if ((*CollectionRefPtr)->Merge(TempCollection))
@@ -1525,7 +1651,7 @@ void FCollectionManager::LoadCollections()
 			const bool bUseSCC = ShouldUseSCC(ShareType);
 
 			FText LoadErrorText;
-			TSharedRef<FCollection> NewCollection = MakeShareable(new FCollection(Filename, bUseSCC));
+			TSharedRef<FCollection> NewCollection = MakeShareable(new FCollection(Filename, bUseSCC, ECollectionStorageMode::Static));
 			if (NewCollection->Load(LoadErrorText))
 			{
 				AddCollection(NewCollection, ShareType);
diff --git a/Engine/Source/Developer/CollectionManager/Private/CollectionManager.h b/Engine/Source/Developer/CollectionManager/Private/CollectionManager.h
--- a/Engine/Source/Developer/CollectionManager/Private/CollectionManager.h
+++ b/Engine/Source/Developer/CollectionManager/Private/CollectionManager.h
@@ -114,7 +114,7 @@ public:
 	virtual FString GetCollectionsStringForObject(FName ObjectPath, ECollectionShareType::Type ShareType, ECollectionRecursionFlags::Flags RecursionMode = ECollectionRecursionFlags::Self) const override;
 	virtual void CreateUniqueCollectionName(const FName& BaseName, ECollectionShareType::Type ShareType, FName& OutCollectionName) const override;
 	virtual bool IsValidCollectionName(const FString& CollectionName, ECollectionShareType::Type ShareType) const override;
-	virtual bool CreateCollection(FName CollectionName, ECollectionShareType::Type ShareType) override;
+	virtual bool CreateCollection(FName CollectionName, ECollectionShareType::Type ShareType, ECollectionStorageMode::Type StorageMode) override;
 	virtual bool RenameCollection(FName CurrentCollectionName, ECollectionShareType::Type CurrentShareType, FName NewCollectionName, ECollectionShareType::Type NewShareType) override;
 	virtual bool ReparentCollection(FName CollectionName, ECollectionShareType::Type ShareType, FName ParentCollectionName, ECollectionShareType::Type ParentShareType) override;
 	virtual bool DestroyCollection(FName CollectionName, ECollectionShareType::Type ShareType) override;
@@ -122,10 +122,13 @@ public:
 	virtual bool AddToCollection(FName CollectionName, ECollectionShareType::Type ShareType, const TArray<FName>& ObjectPaths, int32* OutNumAdded = nullptr) override;
 	virtual bool RemoveFromCollection(FName CollectionName, ECollectionShareType::Type ShareType, FName ObjectPath) override;
 	virtual bool RemoveFromCollection(FName CollectionName, ECollectionShareType::Type ShareType, const TArray<FName>& ObjectPaths, int32* OutNumRemoved = nullptr) override;
+	virtual bool SetDynamicQueryText(FName CollectionName, ECollectionShareType::Type ShareType, const FString& InQueryText) override;
+	virtual bool GetDynamicQueryText(FName CollectionName, ECollectionShareType::Type ShareType, FString& OutQueryText) const override;
 	virtual bool EmptyCollection(FName CollectionName, ECollectionShareType::Type ShareType) override;
 	virtual bool SaveCollection(FName CollectionName, ECollectionShareType::Type ShareType) override;
 	virtual bool UpdateCollection(FName CollectionName, ECollectionShareType::Type ShareType) override;
 	virtual bool GetCollectionStatusInfo(FName CollectionName, ECollectionShareType::Type ShareType, FCollectionStatusInfo& OutStatusInfo) const override;
+	virtual bool GetCollectionStorageMode(FName CollectionName, ECollectionShareType::Type ShareType, ECollectionStorageMode::Type& OutStorageMode) const override;
 	virtual bool IsObjectInCollection(FName ObjectPath, FName CollectionName, ECollectionShareType::Type ShareType, ECollectionRecursionFlags::Flags RecursionMode = ECollectionRecursionFlags::Self) const override;
 	virtual bool IsValidParentCollection(FName CollectionName, ECollectionShareType::Type ShareType, FName ParentCollectionName, ECollectionShareType::Type ParentShareType) const override;
 	virtual FText GetLastError() const override { return LastError; }
diff --git a/Engine/Source/Developer/CollectionManager/Private/CollectionManagerConsoleCommands.h b/Engine/Source/Developer/CollectionManager/Private/CollectionManagerConsoleCommands.h
--- a/Engine/Source/Developer/CollectionManager/Private/CollectionManagerConsoleCommands.h
+++ b/Engine/Source/Developer/CollectionManager/Private/CollectionManagerConsoleCommands.h
@@ -43,20 +43,21 @@ public:
 		}
 
 		FName CollectionName = FName(*Args[0]);
+
 		FString ShareStr = Args[1];
-		ECollectionShareType::Type Type;
+		ECollectionShareType::Type ShareType;
 
 		if ( ShareStr == TEXT("LOCAL") )
 		{
-			Type = ECollectionShareType::CST_Local;
+			ShareType = ECollectionShareType::CST_Local;
 		}
 		else if ( ShareStr == TEXT("PRIVATE") )
 		{
-			Type = ECollectionShareType::CST_Private;
+			ShareType = ECollectionShareType::CST_Private;
 		}
 		else if ( ShareStr == TEXT("SHARED") )
 		{
-			Type = ECollectionShareType::CST_Shared;
+			ShareType = ECollectionShareType::CST_Shared;
 		}
 		else
 		{
@@ -64,7 +65,13 @@ public:
 			return;
 		}
 		
-		if ( Module.Get().CreateCollection(CollectionName, Type) )
+		ECollectionStorageMode::Type StorageMode = ECollectionStorageMode::Static;
+		if (Args.Num() >= 3)
+		{
+			StorageMode = ECollectionStorageMode::FromString(*Args[2]);
+		}
+
+		if ( Module.Get().CreateCollection(CollectionName, ShareType, StorageMode) )
 		{
 			UE_LOG(LogCollectionManager, Log, TEXT("Collection created: %s"), *CollectionName.ToString());
 		}
diff --git a/Engine/Source/Developer/CollectionManager/Public/CollectionManagerTypes.h b/Engine/Source/Developer/CollectionManager/Public/CollectionManagerTypes.h
--- a/Engine/Source/Developer/CollectionManager/Public/CollectionManagerTypes.h
+++ b/Engine/Source/Developer/CollectionManager/Public/CollectionManagerTypes.h
@@ -4,6 +4,66 @@
 
 #include "ISourceControlState.h"
 
+struct ECollectionStorageMode
+{
+	enum Type
+	{
+		/** This collection stores a list of contained objects. */
+		Static,
+		/** This collection stores a query that can be run against the Content Browser. This type of collection never directly knows what objects it contains. */
+		Dynamic,
+	};
+
+	static Type FromString(const TCHAR* InString, const Type ReturnIfConversionFails = Static)
+	{
+		if(FPlatformString::Stricmp(InString, TEXT("Static")) == 0)		return Static;
+		if(FPlatformString::Stricmp(InString, TEXT("Dynamic")) == 0)	return Dynamic;
+		return ReturnIfConversionFails;
+	}
+
+	static const TCHAR* ToString(const Type InType)
+	{
+		switch(InType)
+		{
+		case Static:
+			return TEXT("Static");
+		case Dynamic:
+			return TEXT("Dynamic");
+		default:
+			break;
+		}
+		return TEXT("");
+	}
+
+	static FText ToText(const Type InType)
+	{
+		switch(InType)
+		{
+		case Static:
+			return NSLOCTEXT("ECollectionStorageMode", "Static", "Static");
+		case Dynamic:
+			return NSLOCTEXT("ECollectionStorageMode", "Dynamic", "Dynamic");
+		default:
+			break;
+		}
+		return FText::GetEmpty();
+	}
+
+	static FText GetDescription(const Type InType)
+	{
+		switch (InType)
+		{
+		case Static:
+			return NSLOCTEXT("ECollectionStorageMode", "Static_Description", "Static. This collection stores a list of contained objects.");
+		case Dynamic:
+			return NSLOCTEXT("ECollectionStorageMode", "Dynamic_Description", "Dynamic. This collection stores a query that can be run against the Content Browser.");
+		default:
+			break;
+		}
+		return FText::GetEmpty();
+	}
+};
+
 struct ECollectionShareType
 {
 	enum Type
@@ -165,8 +225,8 @@ struct FCollectionNameType
 	friend inline uint32 GetTypeHash( const FCollectionNameType& Key )
 	{
 		uint32 Hash = 0;
-		HashCombine(Hash, GetTypeHash(Key.Name));
-		HashCombine(Hash, GetTypeHash(Key.Type));
+		Hash = HashCombine(Hash, GetTypeHash(Key.Name));
+		Hash = HashCombine(Hash, GetTypeHash(Key.Type));
 		return Hash;
 	}
 
diff --git a/Engine/Source/Developer/CollectionManager/Public/ICollectionManager.h b/Engine/Source/Developer/CollectionManager/Public/ICollectionManager.h
--- a/Engine/Source/Developer/CollectionManager/Public/ICollectionManager.h
+++ b/Engine/Source/Developer/CollectionManager/Public/ICollectionManager.h
@@ -69,9 +69,10 @@ public:
 	  *
 	  * @param CollectionName The name of the new collection
 	  * @param ShareType The way the collection is shared.
+	  * @param StorageMode How does this collection store its objects? (static or dynamic)
 	  * @return true if the add was successful. If false, GetLastError will return a human readable string description of the error.
 	  */
-	virtual bool CreateCollection(FName CollectionName, ECollectionShareType::Type ShareType) = 0;
+	virtual bool CreateCollection(FName CollectionName, ECollectionShareType::Type ShareType, ECollectionStorageMode::Type StorageMode) = 0;
 
 	/**
 	  * Renames a collection. A .collection file will be added to disk and a .collection file will be removed.
@@ -128,6 +129,26 @@ public:
 	virtual bool RemoveFromCollection(FName CollectionName, ECollectionShareType::Type ShareType, FName ObjectPath) = 0;
 	virtual bool RemoveFromCollection(FName CollectionName, ECollectionShareType::Type ShareType, const TArray<FName>& ObjectPaths, int32* OutNumRemoved = nullptr) = 0;
 
+	/**
+	 * Sets the dynamic query text for the specified collection. 
+	 *
+	 * @param CollectionName The collection to set the query on.
+	 * @param ShareType The way the collection is shared.
+	 * @param InQueryText The new query to set.
+	 * @return true if the set was successful. If false, GetLastError will return a human readable string description of the error.
+	 */
+	virtual bool SetDynamicQueryText(FName CollectionName, ECollectionShareType::Type ShareType, const FString& InQueryText) = 0;
+
+	/**
+	 * Gets the dynamic query text for the specified collection. 
+	 *
+	 * @param CollectionName The collection to get the query from.
+	 * @param ShareType The way the collection is shared.
+	 * @param OutQueryText Filled with the query text.
+	 * @return true if the get was successful. If false, GetLastError will return a human readable string description of the error.
+	 */
+	virtual bool GetDynamicQueryText(FName CollectionName, ECollectionShareType::Type ShareType, FString& OutQueryText) const = 0;
+
 	/**
 	  * Removes all assets from the specified collection.
 	  *
@@ -166,10 +187,21 @@ public:
 	  *
 	  * @param CollectionName The collection to get the status info for
 	  * @param ShareType The way the collection is shared.
+	  * @param OutStatusInfo The status info to populate.
 	  * @return true if the status info was filled in. If false, GetLastError will return a human readable string description of the error.
 	  */
 	virtual bool GetCollectionStatusInfo(FName CollectionName, ECollectionShareType::Type ShareType, FCollectionStatusInfo& OutStatusInfo) const = 0;
 
+	/**
+	 * Gets the method by which the specified collection stores its objects (static or dynamic)
+	 *
+	 * @param CollectionName The collection to get the storage mode for
+	 * @param ShareType The way the collection is shared.
+	 * @param OutStorageMode The variable to populate.
+	 * @return true if the status info was filled in. If false, GetLastError will return a human readable string description of the error.
+	 */
+	virtual bool GetCollectionStorageMode(FName CollectionName, ECollectionShareType::Type ShareType, ECollectionStorageMode::Type& OutStorageMode) const = 0;
+
 	/**
 	 * Check to see if the given object exists in the given collection
 	 *
diff --git a/Engine/Source/Editor/Blutility/Private/BlutilityShelf.cpp b/Engine/Source/Editor/Blutility/Private/BlutilityShelf.cpp
--- a/Engine/Source/Editor/Blutility/Private/BlutilityShelf.cpp
+++ b/Engine/Source/Editor/Blutility/Private/BlutilityShelf.cpp
@@ -68,7 +68,7 @@ void SBlutilityShelf::Construct(const FArguments& InArgs)
 	else
 	{
 		// Create the collection
-		CollectionManagerModule.Get().CreateCollection(BlutilityModule::BlutilityShelfCollectionName, ECollectionShareType::CST_Local);
+		CollectionManagerModule.Get().CreateCollection(BlutilityModule::BlutilityShelfCollectionName, ECollectionShareType::CST_Local, ECollectionStorageMode::Static);
 	}
 
 	// Build the shelf
diff --git a/Engine/Source/Editor/ContentBrowser/Private/AssetContextMenu.cpp b/Engine/Source/Editor/ContentBrowser/Private/AssetContextMenu.cpp
--- a/Engine/Source/Editor/ContentBrowser/Private/AssetContextMenu.cpp
+++ b/Engine/Source/Editor/ContentBrowser/Private/AssetContextMenu.cpp
@@ -847,6 +847,14 @@ bool FAssetContextMenu::AddCollectionMenuOptions(FMenuBuilder& MenuBuilder)
 					continue;
 				}
 
+				// Can only manage assets for static collections
+				ECollectionStorageMode::Type StorageMode = ECollectionStorageMode::Static;
+				CollectionManagerModule.Get().GetCollectionStorageMode(AvailableCollection.Name, AvailableCollection.Type, StorageMode);
+				if (StorageMode != ECollectionStorageMode::Static)
+				{
+					continue;
+				}
+
 				TArray<FCollectionNameType> AvailableChildCollections;
 				CollectionManagerModule.Get().GetChildCollections(AvailableCollection.Name, AvailableCollection.Type, AvailableChildCollections);
 
@@ -854,7 +862,7 @@ bool FAssetContextMenu::AddCollectionMenuOptions(FMenuBuilder& MenuBuilder)
 				{
 					SubMenuBuilder.AddSubMenu(
 						FText::FromName(AvailableCollection.Name), 
-						ECollectionShareType::GetDescription(AvailableCollection.Type), 
+						FText::GetEmpty(), 
 						FNewMenuDelegate::CreateStatic(&FManageCollectionsContextMenu::CreateManageCollectionsSubMenu, QuickAssetManagement, AvailableChildCollections),
 						FUIAction(
 							FExecuteAction::CreateStatic(&FManageCollectionsContextMenu::OnCollectionClicked, QuickAssetManagement, AvailableCollection),
@@ -871,7 +879,7 @@ bool FAssetContextMenu::AddCollectionMenuOptions(FMenuBuilder& MenuBuilder)
 				{
 					SubMenuBuilder.AddMenuEntry(
 						FText::FromName(AvailableCollection.Name), 
-						ECollectionShareType::GetDescription(AvailableCollection.Type), 
+						FText::GetEmpty(), 
 						FSlateIcon(FEditorStyle::GetStyleSetName(), ECollectionShareType::GetIconStyleName(AvailableCollection.Type)), 
 						FUIAction(
 							FExecuteAction::CreateStatic(&FManageCollectionsContextMenu::OnCollectionClicked, QuickAssetManagement, AvailableCollection),
@@ -933,7 +941,7 @@ bool FAssetContextMenu::AddCollectionMenuOptions(FMenuBuilder& MenuBuilder)
 	}
 
 	// "Remove from collection" (only display option if exactly one collection is selected)
-	if ( SourcesData.Collections.Num() == 1 )
+	if ( SourcesData.Collections.Num() == 1 && !SourcesData.IsDynamicCollection() )
 	{
 		MenuBuilder.AddMenuEntry(
 			FText::Format(LOCTEXT("RemoveFromCollectionFmt", "Remove From {0}"), FText::FromName(SourcesData.Collections[0].Name)),
@@ -1572,9 +1580,8 @@ void FAssetContextMenu::ExecuteRemoveFromCollection()
 		{
 			FCollectionManagerModule& CollectionManagerModule = FCollectionManagerModule::GetModule();
 
-			FName CollectionName = SourcesData.Collections[0].Name;
-			ECollectionShareType::Type CollectionType = SourcesData.Collections[0].Type;
-			CollectionManagerModule.Get().RemoveFromCollection(CollectionName, CollectionType, AssetsToRemove);
+			const FCollectionNameType& Collection = SourcesData.Collections[0];
+			CollectionManagerModule.Get().RemoveFromCollection(Collection.Name, Collection.Type, AssetsToRemove);
 			OnAssetViewRefreshRequested.ExecuteIfBound();
 		}
 	}
@@ -1879,7 +1886,7 @@ bool FAssetContextMenu::CanExecuteDelete() const
 
 bool FAssetContextMenu::CanExecuteRemoveFromCollection() const 
 {
-	return SourcesData.Collections.Num() == 1;
+	return SourcesData.Collections.Num() == 1 && !SourcesData.IsDynamicCollection();
 }
 
 bool FAssetContextMenu::CanExecuteSCCRefresh() const
diff --git a/Engine/Source/Editor/ContentBrowser/Private/CollectionContextMenu.cpp b/Engine/Source/Editor/ContentBrowser/Private/CollectionContextMenu.cpp
--- a/Engine/Source/Editor/ContentBrowser/Private/CollectionContextMenu.cpp
+++ b/Engine/Source/Editor/ContentBrowser/Private/CollectionContextMenu.cpp
@@ -68,15 +68,16 @@ TSharedPtr<SWidget> FCollectionContextMenu::MakeCollectionTreeContextMenu(TShare
 	MenuBuilder.BeginSection("CollectionOptions", LOCTEXT("CollectionListOptionsMenuHeading", "Collection Options"));
 	{
 		const bool bHasSingleSelectedCollection = SelectedCollections.Num() == 1;
+		const bool bIsFirstSelectedCollectionStatic = SelectedCollections.Num() > 0 && SelectedCollections[0]->StorageMode == ECollectionStorageMode::Static;
 
 		// New... (submenu)
 		MenuBuilder.AddSubMenu(
 			LOCTEXT("NewChildCollection", "New..."),
 			LOCTEXT("NewChildCollectionTooltip", "Create child a collection."),
-			FNewMenuDelegate::CreateRaw( this, &FCollectionContextMenu::MakeNewCollectionSubMenu, TOptional<FCollectionNameType>( FCollectionNameType( SelectedCollections[0]->CollectionName, SelectedCollections[0]->CollectionType ) ) ),
+			FNewMenuDelegate::CreateRaw( this, &FCollectionContextMenu::MakeNewCollectionSubMenu, ECollectionStorageMode::Static, SCollectionView::FCreateCollectionPayload( FCollectionNameType( SelectedCollections[0]->CollectionName, SelectedCollections[0]->CollectionType ) ) ),
 			FUIAction(
 				FExecuteAction(),
-				FCanExecuteAction::CreateLambda( [=]{ return bHasSingleSelectedCollection; } )
+				FCanExecuteAction::CreateLambda( [=]{ return bHasSingleSelectedCollection && bIsFirstSelectedCollectionStatic; } )
 				),
 			NAME_None,
 			EUserInterfaceActionType::Button
@@ -164,41 +165,44 @@ TSharedPtr<SWidget> FCollectionContextMenu::MakeCollectionTreeContextMenu(TShare
 	return MenuBuilder.MakeWidget();
 }
 
-void FCollectionContextMenu::MakeNewCollectionSubMenu(FMenuBuilder& MenuBuilder, TOptional<FCollectionNameType> ParentCollection)
+void FCollectionContextMenu::MakeNewCollectionSubMenu(FMenuBuilder& MenuBuilder, ECollectionStorageMode::Type StorageMode, SCollectionView::FCreateCollectionPayload InCreationPayload)
 {
-	const FText MenuHeading = (ParentCollection.IsSet()) ? LOCTEXT("NewChildCollectionMenuHeading", "New Child Collection"): LOCTEXT("NewCollectionMenuHeading", "New Collection");
+	const FText MenuHeading = FText::Format(
+		(InCreationPayload.ParentCollection.IsSet()) ? LOCTEXT("NewXChildCollectionMenuHeading", "New {0} Child Collection") : LOCTEXT("NewXCollectionMenuHeading", "New {0} Collection"),
+		ECollectionStorageMode::ToText(StorageMode)
+		);
 
 	MenuBuilder.BeginSection("CollectionNewCollection", MenuHeading);
 	{
-		const bool bCanCreateSharedChildren = !ParentCollection.IsSet() || ECollectionShareType::IsValidChildType( ParentCollection->Type, ECollectionShareType::CST_Shared );
+		const bool bCanCreateSharedChildren = !InCreationPayload.ParentCollection.IsSet() || ECollectionShareType::IsValidChildType( InCreationPayload.ParentCollection->Type, ECollectionShareType::CST_Shared );
 		MenuBuilder.AddMenuEntry(
 			LOCTEXT("NewCollection_Shared", "Shared Collection"),
 			LOCTEXT("NewCollection_SharedTooltip", "Create a collection that can be seen by anyone."),
 			FSlateIcon( FEditorStyle::GetStyleSetName(), ECollectionShareType::GetIconStyleName( ECollectionShareType::CST_Shared ) ),
 			FUIAction(
-				FExecuteAction::CreateSP( this, &FCollectionContextMenu::ExecuteNewCollection, ECollectionShareType::CST_Shared, ParentCollection ),
+				FExecuteAction::CreateSP( this, &FCollectionContextMenu::ExecuteNewCollection, ECollectionShareType::CST_Shared, StorageMode, InCreationPayload ),
 				FCanExecuteAction::CreateSP( this, &FCollectionContextMenu::CanExecuteNewCollection, ECollectionShareType::CST_Shared, bCanCreateSharedChildren )
 				)
 			);
 
-		const bool bCanCreatePrivateChildren = !ParentCollection.IsSet() || ECollectionShareType::IsValidChildType( ParentCollection->Type, ECollectionShareType::CST_Private );
+		const bool bCanCreatePrivateChildren = !InCreationPayload.ParentCollection.IsSet() || ECollectionShareType::IsValidChildType( InCreationPayload.ParentCollection->Type, ECollectionShareType::CST_Private );
 		MenuBuilder.AddMenuEntry(
 			LOCTEXT("NewCollection_Private", "Private Collection"),
 			LOCTEXT("NewCollection_PrivateTooltip", "Create a collection that can only be seen by you."),
 			FSlateIcon( FEditorStyle::GetStyleSetName(), ECollectionShareType::GetIconStyleName( ECollectionShareType::CST_Private ) ),
 			FUIAction(
-				FExecuteAction::CreateSP( this, &FCollectionContextMenu::ExecuteNewCollection, ECollectionShareType::CST_Private, ParentCollection ),
+				FExecuteAction::CreateSP( this, &FCollectionContextMenu::ExecuteNewCollection, ECollectionShareType::CST_Private, StorageMode, InCreationPayload ),
 				FCanExecuteAction::CreateSP( this, &FCollectionContextMenu::CanExecuteNewCollection, ECollectionShareType::CST_Private, bCanCreatePrivateChildren )
 				)
 			);
 
-		const bool bCanCreateLocalChildren = !ParentCollection.IsSet() || ECollectionShareType::IsValidChildType( ParentCollection->Type, ECollectionShareType::CST_Local );
+		const bool bCanCreateLocalChildren = !InCreationPayload.ParentCollection.IsSet() || ECollectionShareType::IsValidChildType( InCreationPayload.ParentCollection->Type, ECollectionShareType::CST_Local );
 		MenuBuilder.AddMenuEntry(
 			LOCTEXT("NewCollection_Local", "Local Collection"),
 			LOCTEXT("NewCollection_LocalTooltip", "Create a collection that is not in source control and can only be seen by you."),
 			FSlateIcon( FEditorStyle::GetStyleSetName(), ECollectionShareType::GetIconStyleName( ECollectionShareType::CST_Local ) ),
 			FUIAction(
-				FExecuteAction::CreateSP( this, &FCollectionContextMenu::ExecuteNewCollection, ECollectionShareType::CST_Local, ParentCollection ),
+				FExecuteAction::CreateSP( this, &FCollectionContextMenu::ExecuteNewCollection, ECollectionShareType::CST_Local, StorageMode, InCreationPayload ),
 				FCanExecuteAction::CreateSP( this, &FCollectionContextMenu::CanExecuteNewCollection, ECollectionShareType::CST_Local, bCanCreateLocalChildren )
 				)
 			);
@@ -206,6 +210,58 @@ void FCollectionContextMenu::MakeNewCollectionSubMenu(FMenuBuilder& MenuBuilder,
 	MenuBuilder.EndSection();
 }
 
+void FCollectionContextMenu::MakeSaveDynamicCollectionSubMenu(FMenuBuilder& MenuBuilder, FText InSearchQuery)
+{
+	auto OnCollectionCreated = FCollectionItem::FCollectionCreatedEvent::CreateSP(this, &FCollectionContextMenu::ExecuteSaveDynamicCollection, InSearchQuery);
+
+	// Create new root level collection
+	MakeNewCollectionSubMenu(MenuBuilder, ECollectionStorageMode::Dynamic, SCollectionView::FCreateCollectionPayload(OnCollectionCreated));
+
+	FCollectionManagerModule& CollectionManagerModule = FCollectionManagerModule::GetModule();
+	
+	TArray<FCollectionNameType> AvailableCollections;
+	CollectionManagerModule.Get().GetCollections(AvailableCollections);
+
+	AvailableCollections.Sort([](const FCollectionNameType& One, const FCollectionNameType& Two) -> bool
+	{
+		return One.Name < Two.Name;
+	});
+
+	if (AvailableCollections.Num() > 0)
+	{
+		MenuBuilder.BeginSection("CollectionReplaceCollection", LOCTEXT("OverwriteDynamicCollectionMenuHeading", "Overwrite Dynamic Collection"));
+
+		for (const FCollectionNameType& AvailableCollection : AvailableCollections)
+		{
+			// Never display system collections
+			if (AvailableCollection.Type == ECollectionShareType::CST_System)
+			{
+				continue;
+			}
+
+			// Can only overwrite dynamic collections
+			ECollectionStorageMode::Type StorageMode = ECollectionStorageMode::Static;
+			CollectionManagerModule.Get().GetCollectionStorageMode(AvailableCollection.Name, AvailableCollection.Type, StorageMode);
+			if (StorageMode != ECollectionStorageMode::Dynamic)
+			{
+				continue;
+			}
+
+			MenuBuilder.AddMenuEntry(
+				FText::FromName(AvailableCollection.Name), 
+				FText::Format(LOCTEXT("SaveDynamicCollection_OverwriteExistingCollectionToolTip", "Overwrite '{0}' with the current search query"), FText::FromName(AvailableCollection.Name)),
+				FSlateIcon(FEditorStyle::GetStyleSetName(), ECollectionShareType::GetIconStyleName(AvailableCollection.Type)),
+				FUIAction(
+					FExecuteAction::CreateSP( this, &FCollectionContextMenu::ExecuteSaveDynamicCollection, AvailableCollection, InSearchQuery ),
+					FCanExecuteAction::CreateSP( this, &FCollectionContextMenu::CanExecuteSaveDynamicCollection, AvailableCollection )
+					)
+				);
+		}
+
+		MenuBuilder.EndSection();
+	}
+}
+
 void FCollectionContextMenu::MakeCollectionShareTypeSubMenu(FMenuBuilder& MenuBuilder)
 {
 	MenuBuilder.BeginSection("CollectionShareType", LOCTEXT("CollectionShareTypeMenuHeading", "Collection Share Type"));
@@ -329,14 +385,14 @@ bool FCollectionContextMenu::CanRenameSelectedCollections() const
 	return false;
 }
 
-void FCollectionContextMenu::ExecuteNewCollection(ECollectionShareType::Type CollectionType, TOptional<FCollectionNameType> ParentCollection)
+void FCollectionContextMenu::ExecuteNewCollection(ECollectionShareType::Type CollectionType, ECollectionStorageMode::Type StorageMode, SCollectionView::FCreateCollectionPayload InCreationPayload)
 {
 	if ( !ensure(CollectionView.IsValid()) )
 	{
 		return;
 	}
 
-	CollectionView.Pin()->CreateCollectionItem(CollectionType, ParentCollection);
+	CollectionView.Pin()->CreateCollectionItem(CollectionType, StorageMode, InCreationPayload);
 }
 
 void FCollectionContextMenu::ExecuteSetCollectionShareType(ECollectionShareType::Type CollectionType)
@@ -358,6 +414,13 @@ void FCollectionContextMenu::ExecuteSetCollectionShareType(ECollectionShareType:
 	CollectionManagerModule.Get().RenameCollection(SelectedCollections[0]->CollectionName, SelectedCollections[0]->CollectionType, SelectedCollections[0]->CollectionName, CollectionType);
 }
 
+void FCollectionContextMenu::ExecuteSaveDynamicCollection(FCollectionNameType InCollection, FText InSearchQuery)
+{
+	FCollectionManagerModule& CollectionManagerModule = FCollectionManagerModule::GetModule();
+
+	CollectionManagerModule.Get().SetDynamicQueryText(InCollection.Name, InCollection.Type, InSearchQuery.ToString());
+}
+
 void FCollectionContextMenu::ExecuteRenameCollection()
 {
 	if ( !ensure(CollectionView.IsValid()) )
@@ -529,6 +592,11 @@ bool FCollectionContextMenu::IsSetCollectionShareTypeChecked(ECollectionShareTyp
 	return SelectedCollections[0]->CollectionType == CollectionType;
 }
 
+bool FCollectionContextMenu::CanExecuteSaveDynamicCollection(FCollectionNameType InCollection) const
+{
+	return InCollection.Type == ECollectionShareType::CST_Local || (bProjectUnderSourceControl && ISourceControlModule::Get().IsEnabled() && ISourceControlModule::Get().GetProvider().IsAvailable());
+}
+
 bool FCollectionContextMenu::CanExecuteRenameCollection() const
 {
 	return CanRenameSelectedCollections();
diff --git a/Engine/Source/Editor/ContentBrowser/Private/CollectionContextMenu.h b/Engine/Source/Editor/ContentBrowser/Private/CollectionContextMenu.h
--- a/Engine/Source/Editor/ContentBrowser/Private/CollectionContextMenu.h
+++ b/Engine/Source/Editor/ContentBrowser/Private/CollectionContextMenu.h
@@ -15,7 +15,10 @@ public:
 	TSharedPtr<SWidget> MakeCollectionTreeContextMenu(TSharedPtr< FUICommandList > InCommandList);
 
 	/** Makes the new collection submenu */
-	void MakeNewCollectionSubMenu(FMenuBuilder& MenuBuilder, TOptional<FCollectionNameType> ParentCollection);
+	void MakeNewCollectionSubMenu(FMenuBuilder& MenuBuilder, ECollectionStorageMode::Type StorageMode, SCollectionView::FCreateCollectionPayload InCreationPayload);
+
+	/** Makes the save dynamic collection submenu */
+	void MakeSaveDynamicCollectionSubMenu(FMenuBuilder& MenuBuilder, FText InSearchQuery);
 
 	/** Makes the collection share type submenu */
 	void MakeCollectionShareTypeSubMenu(FMenuBuilder& MenuBuilder);
@@ -32,11 +35,14 @@ protected:
 
 private:
 	/** Handler for when a collection is selected in the "New" menu */
-	void ExecuteNewCollection(ECollectionShareType::Type CollectionType, TOptional<FCollectionNameType> ParentCollection);
+	void ExecuteNewCollection(ECollectionShareType::Type CollectionType, ECollectionStorageMode::Type StorageMode, SCollectionView::FCreateCollectionPayload InCreationPayload);
 
 	/** Handler for when a collection share type is changed in the "Share Type" menu */
 	void ExecuteSetCollectionShareType(ECollectionShareType::Type CollectionType);
 
+	/** Handler for when a dynamic collection is selected in the "Save" menu */
+	void ExecuteSaveDynamicCollection(FCollectionNameType InCollection, FText InSearchQuery);
+
 	/** Handler for when "Rename Collection" is selected */
 	void ExecuteRenameCollection();
 
@@ -67,6 +73,9 @@ private:
 	/** Handler to check to see if an entry in the "Share Type" menu should be checked */
 	bool IsSetCollectionShareTypeChecked(ECollectionShareType::Type CollectionType) const;
 
+	/** Handler to check to see if "Save Dynamic Collection" can be executed */
+	bool CanExecuteSaveDynamicCollection(FCollectionNameType InCollection) const;
+
 	/** Handler to check to see if "Rename Collection" can be executed */
 	bool CanExecuteRenameCollection() const;
 
diff --git a/Engine/Source/Editor/ContentBrowser/Private/CollectionViewTypes.h b/Engine/Source/Editor/ContentBrowser/Private/CollectionViewTypes.h
--- a/Engine/Source/Editor/ContentBrowser/Private/CollectionViewTypes.h
+++ b/Engine/Source/Editor/ContentBrowser/Private/CollectionViewTypes.h
@@ -44,6 +44,9 @@ struct FCollectionItem
 	/** The type of the collection */
 	ECollectionShareType::Type CollectionType;
 
+	/** How does this collection store its objects? (static or dynamic) */
+	ECollectionStorageMode::Type StorageMode;
+
 	/** Pointer to our parent collection (if any) */
 	TWeakPtr<FCollectionItem> ParentCollection;
 
@@ -59,16 +62,23 @@ struct FCollectionItem
 	/** Current status of this collection item */
 	ECollectionItemStatus CurrentStatus;
 
+	/** Called once after the collection is created (see bNewCollection) */
+	DECLARE_DELEGATE_OneParam( FCollectionCreatedEvent, FCollectionNameType )
+
 	/** Broadcasts whenever renaming a collection item is requested */
 	DECLARE_MULTICAST_DELEGATE( FRenamedRequestEvent )
 
+	/** Broadcasts once after the collection is created (see bNewCollection) */
+	FCollectionCreatedEvent OnCollectionCreatedEvent;
+
 	/** Broadcasts whenever a rename is requested */
 	FRenamedRequestEvent OnRenamedRequestEvent;
 
 	/** Constructor */
-	FCollectionItem(const FName& InCollectionName, const ECollectionShareType::Type& InCollectionType)
+	FCollectionItem(const FName& InCollectionName, const ECollectionShareType::Type InCollectionType)
 		: CollectionName(InCollectionName)
 		, CollectionType(InCollectionType)
+		, StorageMode(ECollectionStorageMode::Static)
 		, ParentCollection()
 		, ChildCollections()
 		, bRenaming(false)
diff --git a/Engine/Source/Editor/ContentBrowser/Private/SAssetPicker.cpp b/Engine/Source/Editor/ContentBrowser/Private/SAssetPicker.cpp
--- a/Engine/Source/Editor/ContentBrowser/Private/SAssetPicker.cpp
+++ b/Engine/Source/Editor/ContentBrowser/Private/SAssetPicker.cpp
@@ -166,10 +166,9 @@ void SAssetPicker::Construct( const FArguments& InArgs )
 	// Asset view
 	
 	// Break up the incoming filter into a sources data and backend filter.
-	CurrentSourcesData.PackagePaths = InArgs._AssetPickerConfig.Filter.PackagePaths;
-	CurrentSourcesData.Collections = InArgs._AssetPickerConfig.Collections;
+	CurrentSourcesData = FSourcesData(InArgs._AssetPickerConfig.Filter.PackagePaths, InArgs._AssetPickerConfig.Collections);
 	CurrentBackendFilter = InArgs._AssetPickerConfig.Filter;
-	CurrentBackendFilter.PackagePaths.Empty();
+	CurrentBackendFilter.PackagePaths.Reset();
 
 	if(InArgs._AssetPickerConfig.bAddFilterUI)
 	{
@@ -289,7 +288,7 @@ FReply SAssetPicker::OnKeyDown(const FGeometry& MyGeometry, const FKeyEvent& InK
 
 void SAssetPicker::FolderEntered(const FString& FolderPath)
 {
-	CurrentSourcesData.PackagePaths.Empty();
+	CurrentSourcesData.PackagePaths.Reset();
 	CurrentSourcesData.PackagePaths.Add(FName(*FolderPath));
 
 	AssetViewPtr->SetSourcesData(CurrentSourcesData);
@@ -348,7 +347,7 @@ void SAssetPicker::SetNewBackendFilter(const FARFilter& NewFilter)
 	AssetViewPtr->SetSourcesData(CurrentSourcesData);
 
 	CurrentBackendFilter = NewFilter;
-	CurrentBackendFilter.PackagePaths.Empty();
+	CurrentBackendFilter.PackagePaths.Reset();
 
 	// Update the Text filter too, since now class names may no longer matter
 	TextFilter->SetIncludeClassName(NewFilter.ClassNames.Num() != 1);
diff --git a/Engine/Source/Editor/ContentBrowser/Private/SAssetView.cpp b/Engine/Source/Editor/ContentBrowser/Private/SAssetView.cpp
--- a/Engine/Source/Editor/ContentBrowser/Private/SAssetView.cpp
+++ b/Engine/Source/Editor/ContentBrowser/Private/SAssetView.cpp
@@ -1111,7 +1111,8 @@ FReply SAssetView::OnDragOver( const FGeometry& MyGeometry, const FDragDropEvent
 			{
 				TArray< FName > ObjectPaths;
 				FCollectionManagerModule& CollectionManagerModule = FCollectionManagerModule::GetModule();
-				CollectionManagerModule.Get().GetObjectsInCollection( SourcesData.Collections[0].Name, SourcesData.Collections[0].Type, ObjectPaths );
+				const FCollectionNameType& Collection = SourcesData.Collections[0];
+				CollectionManagerModule.Get().GetObjectsInCollection( Collection.Name, Collection.Type, ObjectPaths );
 
 				bool IsValidDrop = false;
 				for (const auto& AssetData : AssetDatas)
@@ -1176,7 +1177,8 @@ FReply SAssetView::OnDrop( const FGeometry& MyGeometry, const FDragDropEvent& Dr
 			if (ObjectPaths.Num() > 0)
 			{
 				FCollectionManagerModule& CollectionManagerModule = FCollectionManagerModule::GetModule();
-				CollectionManagerModule.Get().AddToCollection(SourcesData.Collections[0].Name, SourcesData.Collections[0].Type, ObjectPaths);
+				const FCollectionNameType& Collection = SourcesData.Collections[0];
+				CollectionManagerModule.Get().AddToCollection(Collection.Name, Collection.Type, ObjectPaths);
 			}
 
 			return FReply::Handled();
@@ -1341,6 +1343,7 @@ void SAssetView::RefreshSourceItems()
 		// force recursion when the user is searching
 		const bool bRecurse = ShouldFilterRecursively();
 		const bool bUsingFolders = IsShowingFolders();
+		const bool bIsDynamicCollection = SourcesData.IsDynamicCollection();
 		FARFilter Filter = SourcesData.MakeFilter(bRecurse, bUsingFolders);
 
 		// Add the backend filters from the filter list
@@ -1361,9 +1364,9 @@ void SAssetView::RefreshSourceItems()
 
 		// Only show classes if we have class paths, and the filter allows classes to be shown
 		const bool bFilterAllowsClasses = Filter.ClassNames.Num() == 0 || Filter.ClassNames.Contains(NAME_Class);
-		bShowClasses = ClassPathsToShow.Num() > 0 && bFilterAllowsClasses;
+		bShowClasses = (ClassPathsToShow.Num() > 0 || bIsDynamicCollection) && bFilterAllowsClasses;
 
-		if ( SourcesData.Collections.Num() > 0 && Filter.ObjectPaths.Num() == 0 )
+		if ( SourcesData.HasCollections() && Filter.ObjectPaths.Num() == 0 && !bIsDynamicCollection )
 		{
 			// This is an empty collection, no asset will pass the check
 		}
@@ -1379,19 +1382,20 @@ void SAssetView::RefreshSourceItems()
 
 		if ( bFilterAllowsClasses )
 		{
+			FCollectionManagerModule& CollectionManagerModule = FCollectionManagerModule::GetModule();
+
 			// Include objects from child collections if we're recursing
 			const ECollectionRecursionFlags::Flags CollectionRecursionMode = (Filter.bRecursivePaths) ? ECollectionRecursionFlags::SelfAndChildren : ECollectionRecursionFlags::Self;
 
 			TArray< FName > ClassPaths;
-			FCollectionManagerModule& CollectionManagerModule = FCollectionManagerModule::GetModule();
-			for (int32 Index = 0; Index < SourcesData.Collections.Num(); Index++)
+			for (const FCollectionNameType& Collection : SourcesData.Collections)
 			{
-				CollectionManagerModule.Get().GetClassesInCollection( SourcesData.Collections[Index].Name, SourcesData.Collections[Index].Type, ClassPaths, CollectionRecursionMode );
+				CollectionManagerModule.Get().GetClassesInCollection( Collection.Name, Collection.Type, ClassPaths, CollectionRecursionMode );
 			}
 
-			for (int32 Index = 0; Index < ClassPaths.Num(); Index++)
+			for (const FName& ClassPath : ClassPaths)
 			{
-				UClass* Class = FindObject<UClass>(ANY_PACKAGE, *ClassPaths[Index].ToString());
+				UClass* Class = FindObject<UClass>(ANY_PACKAGE, *ClassPath.ToString());
 
 				if ( Class != NULL )
 				{
@@ -1711,7 +1715,7 @@ void SAssetView::RefreshFolders()
 	}
 
 	// Add folders for any child collections of the currently selected collections
-	if(SourcesData.Collections.Num() > 0)
+	if(SourcesData.HasCollections())
 	{
 		FCollectionManagerModule& CollectionManagerModule = FCollectionManagerModule::GetModule();
 		
@@ -2014,9 +2018,9 @@ void SAssetView::OnAssetRegistryPathAdded(const FString& Path)
 		const bool bDisplayDev = GetDefault<UContentBrowserSettings>()->GetDisplayDevelopersFolder();
 		if ( bDisplayDev || !ContentBrowserUtils::IsDevelopersFolder(Path) )
 		{
-			for(auto SourcePathIt(SourcesData.PackagePaths.CreateConstIterator()); SourcePathIt; SourcePathIt++)
+			for (const FName& SourcePathName : SourcesData.PackagePaths)
 			{
-				const FString SourcePath = (*SourcePathIt).ToString();
+				const FString SourcePath = SourcePathName.ToString();
 				if(Path.StartsWith(SourcePath))
 				{
 					const FString SubPath = Path.RightChop(SourcePath.Len());
@@ -2192,9 +2196,10 @@ void SAssetView::RunAssetsThroughBackendFilter(TArray<FAssetData>& InOutAssetDat
 {
 	const bool bRecurse = ShouldFilterRecursively();
 	const bool bUsingFolders = IsShowingFolders();
+	const bool bIsDynamicCollection = SourcesData.IsDynamicCollection();
 	FARFilter Filter = SourcesData.MakeFilter(bRecurse, bUsingFolders);
 	
-	if ( SourcesData.Collections.Num() > 0 && Filter.ObjectPaths.Num() == 0 )
+	if ( SourcesData.HasCollections() && Filter.ObjectPaths.Num() == 0 && !bIsDynamicCollection )
 	{
 		// This is an empty collection, no asset will pass the check
 		InOutAssetDataList.Empty();
@@ -2207,7 +2212,7 @@ void SAssetView::RunAssetsThroughBackendFilter(TArray<FAssetData>& InOutAssetDat
 		FAssetRegistryModule& AssetRegistryModule = FModuleManager::LoadModuleChecked<FAssetRegistryModule>(TEXT("AssetRegistry"));
 		AssetRegistryModule.Get().RunAssetsThroughFilter(InOutAssetDataList, Filter);
 
-		if ( SourcesData.Collections.Num() > 0 )
+		if ( SourcesData.HasCollections() && !bIsDynamicCollection )
 		{
 			// Include objects from child collections if we're recursing
 			const ECollectionRecursionFlags::Flags CollectionRecursionMode = (Filter.bRecursivePaths) ? ECollectionRecursionFlags::SelfAndChildren : ECollectionRecursionFlags::Self;
@@ -3219,7 +3224,7 @@ bool SAssetView::CanOpenContextMenu() const
 		return false;
 	}
 
-	if ( SelectedAssets.Num() == 0 && SourcesData.Collections.Num() > 0 )
+	if ( SelectedAssets.Num() == 0 && SourcesData.HasCollections() )
 	{
 		// Don't allow a context menu when we're viewing a collection and have no assets selected
 		return false;
@@ -3313,7 +3318,7 @@ FReply SAssetView::OnDraggingAssetItem( const FGeometry& MyGeometry, const FPoin
 		{
 			// are we dragging some folders?
 			TArray<FString> SelectedFolders = GetSelectedFolders();
-			if(SelectedFolders.Num() > 0 && SourcesData.Collections.Num() == 0)
+			if(SelectedFolders.Num() > 0 && !SourcesData.HasCollections())
 			{
 				return FReply::Handled().BeginDragDrop(FAssetPathDragDropOp::New(SelectedFolders));
 			}
@@ -3702,7 +3707,7 @@ FText SAssetView::GetAssetShowWarningText() const
 		NothingToShowText = LOCTEXT( "NothingToShowCheckFilter", "No results, check your filter." );
 	}
 
-	if ( SourcesData.Collections.Num() > 0 )
+	if ( SourcesData.HasCollections() && !SourcesData.IsDynamicCollection() )
 	{
 		DropText = LOCTEXT( "DragAssetsHere", "Drag and drop assets here to add them to the collection." );
 	}
diff --git a/Engine/Source/Editor/ContentBrowser/Private/SCollectionView.cpp b/Engine/Source/Editor/ContentBrowser/Private/SCollectionView.cpp
--- a/Engine/Source/Editor/ContentBrowser/Private/SCollectionView.cpp
+++ b/Engine/Source/Editor/ContentBrowser/Private/SCollectionView.cpp
@@ -397,6 +397,8 @@ void SCollectionView::UpdateCollectionItems()
 				TSharedRef<FCollectionItem> CollectionItem = MakeShareable(new FCollectionItem(Collection.Name, Collection.Type));
 				OutAvailableCollections.Add(Collection, CollectionItem);
 
+				CollectionManagerModule.Get().GetCollectionStorageMode(Collection.Name, Collection.Type, CollectionItem->StorageMode);
+
 				SCollectionView::UpdateCollectionItemStatus(CollectionItem);
 
 				if (InParentCollectionItem.IsValid())
@@ -547,7 +549,7 @@ void SCollectionView::SetSelectedCollections(const TArray<FCollectionNameType>&
 			CollectionTreePtr->SetItemSelection(CollectionItemToSelect, true);
 
 			// If the selected collection doesn't pass our current filter, we need to clear it
-			if (bEnsureVisible && !CollectionItemTextFilter->PassesFilter(FCollectionItem(CollectionItemToSelect->CollectionName, CollectionItemToSelect->CollectionType)))
+			if (bEnsureVisible && !CollectionItemTextFilter->PassesFilter(*CollectionItemToSelect))
 			{
 				SearchBoxPtr->SetText(FText::GetEmpty());
 			}
@@ -607,9 +609,9 @@ void SCollectionView::ApplyHistoryData( const FHistoryData& History )
 	FScopedPreventSelectionChangedDelegate DelegatePrevention( SharedThis(this) );
 
 	CollectionTreePtr->ClearSelection();
-	for ( auto HistoryIt = History.SourcesData.Collections.CreateConstIterator(); HistoryIt; ++HistoryIt)
+	for (const FCollectionNameType& HistoryCollection : History.SourcesData.Collections)
 	{
-		TSharedPtr<FCollectionItem> CollectionHistoryItem = AvailableCollections.FindRef(FCollectionNameType((*HistoryIt).Name, (*HistoryIt).Type));
+		TSharedPtr<FCollectionItem> CollectionHistoryItem = AvailableCollections.FindRef(HistoryCollection);
 		if (CollectionHistoryItem.IsValid())
 		{
 			ExpandParentItems(CollectionHistoryItem.ToSharedRef());
@@ -787,6 +789,37 @@ FReply SCollectionView::OnDrop( const FGeometry& MyGeometry, const FDragDropEven
 	return FReply::Unhandled();
 }
 
+void SCollectionView::MakeSaveDynamicCollectionMenu(FText InQueryString)
+{
+	// Get all menu extenders for this context menu from the content browser module
+	FContentBrowserModule& ContentBrowserModule = FModuleManager::GetModuleChecked<FContentBrowserModule>( TEXT("ContentBrowser") );
+	TArray<FContentBrowserMenuExtender> MenuExtenderDelegates = ContentBrowserModule.GetAllCollectionViewContextMenuExtenders();
+
+	TArray<TSharedPtr<FExtender>> Extenders;
+	for (int32 i = 0; i < MenuExtenderDelegates.Num(); ++i)
+	{
+		if (MenuExtenderDelegates[i].IsBound())
+		{
+			Extenders.Add(MenuExtenderDelegates[i].Execute());
+		}
+	}
+	TSharedPtr<FExtender> MenuExtender = FExtender::Combine(Extenders);
+
+	FMenuBuilder MenuBuilder(/*bInShouldCloseWindowAfterMenuSelection=*/true, NULL, MenuExtender, true);
+
+	CollectionContextMenu->UpdateProjectSourceControl();
+
+	CollectionContextMenu->MakeSaveDynamicCollectionSubMenu(MenuBuilder, InQueryString);
+
+	FSlateApplication::Get().PushMenu(
+		AsShared(),
+		FWidgetPath(),
+		MenuBuilder.MakeWidget(),
+		FSlateApplication::Get().GetCursorPos(),
+		FPopupTransitionEffect( FPopupTransitionEffect::TopMenu )
+		);
+}
+
 bool SCollectionView::ShouldAllowSelectionChangedDelegate() const
 {
 	return PreventSelectionChangedDelegateCount == 0;
@@ -812,7 +845,7 @@ FReply SCollectionView::MakeAddCollectionMenu()
 
 	CollectionContextMenu->UpdateProjectSourceControl();
 
-	CollectionContextMenu->MakeNewCollectionSubMenu(MenuBuilder, TOptional<FCollectionNameType>());
+	CollectionContextMenu->MakeNewCollectionSubMenu(MenuBuilder, ECollectionStorageMode::Static, SCollectionView::FCreateCollectionPayload());
 
 	FSlateApplication::Get().PushMenu(
 		AsShared(),
@@ -842,7 +875,7 @@ EVisibility SCollectionView::GetAddCollectionButtonVisibility() const
 	return (bAllowCollectionButtons && ( !CollectionsExpandableAreaPtr.IsValid() || CollectionsExpandableAreaPtr->IsExpanded() ) ) ? EVisibility::Visible : EVisibility::Collapsed;
 }
 
-void SCollectionView::CreateCollectionItem( ECollectionShareType::Type CollectionType, TOptional<FCollectionNameType> ParentCollection )
+void SCollectionView::CreateCollectionItem( ECollectionShareType::Type CollectionType, ECollectionStorageMode::Type StorageMode, const FCreateCollectionPayload& InCreationPayload )
 {
 	if ( ensure(CollectionType != ECollectionShareType::CST_All) )
 	{
@@ -852,13 +885,14 @@ void SCollectionView::CreateCollectionItem( ECollectionShareType::Type Collectio
 		FName CollectionName;
 		CollectionManagerModule.Get().CreateUniqueCollectionName(BaseCollectionName, CollectionType, CollectionName);
 		TSharedPtr<FCollectionItem> NewItem = MakeShareable(new FCollectionItem(CollectionName, CollectionType));
+		NewItem->StorageMode = StorageMode;
 
 		// Adding a new collection now, so clear any filter we may have applied
 		SearchBoxPtr->SetText(FText::GetEmpty());
 
-		if ( ParentCollection.IsSet() )
+		if ( InCreationPayload.ParentCollection.IsSet() )
 		{
-			TSharedPtr<FCollectionItem> ParentCollectionItem = AvailableCollections.FindRef(ParentCollection.GetValue());
+			TSharedPtr<FCollectionItem> ParentCollectionItem = AvailableCollections.FindRef(InCreationPayload.ParentCollection.GetValue());
 			if ( ParentCollectionItem.IsValid() )
 			{
 				ParentCollectionItem->ChildCollections.Add(NewItem);
@@ -872,6 +906,7 @@ void SCollectionView::CreateCollectionItem( ECollectionShareType::Type Collectio
 		// Mark the new collection for rename and that it is new so it will be created upon successful rename
 		NewItem->bRenaming = true;
 		NewItem->bNewCollection = true;
+		NewItem->OnCollectionCreatedEvent = InCreationPayload.OnCollectionCreatedEvent;
 
 		AvailableCollections.Add( FCollectionNameType(NewItem->CollectionName, NewItem->CollectionType), NewItem );
 		UpdateFilteredCollectionItems();
@@ -1040,9 +1075,18 @@ TSharedRef<ITableRow> SCollectionView::GenerateCollectionRow( TSharedPtr<FCollec
 	FOnCheckStateChanged OnCollectionCheckStateChangedDelegate;
 	if ( QuickAssetManagement.IsValid() )
 	{
-		IsCollectionCheckBoxEnabledAttribute.Bind(TAttribute<bool>::FGetter::CreateSP(this, &SCollectionView::IsCollectionCheckBoxEnabled, CollectionItem));
-		IsCollectionCheckedAttribute.Bind(TAttribute<ECheckBoxState>::FGetter::CreateSP(this, &SCollectionView::IsCollectionChecked, CollectionItem));
-		OnCollectionCheckStateChangedDelegate.BindSP(this, &SCollectionView::OnCollectionCheckStateChanged, CollectionItem);
+		// Can only manage assets for static collections
+		if (CollectionItem->StorageMode == ECollectionStorageMode::Static)
+		{
+			IsCollectionCheckBoxEnabledAttribute.Bind(TAttribute<bool>::FGetter::CreateSP(this, &SCollectionView::IsCollectionCheckBoxEnabled, CollectionItem));
+			IsCollectionCheckedAttribute.Bind(TAttribute<ECheckBoxState>::FGetter::CreateSP(this, &SCollectionView::IsCollectionChecked, CollectionItem));
+			OnCollectionCheckStateChangedDelegate.BindSP(this, &SCollectionView::OnCollectionCheckStateChanged, CollectionItem);
+		}
+		else
+		{
+			IsCollectionCheckBoxEnabledAttribute.Bind(TAttribute<bool>::FGetter::CreateLambda([]{ return false; }));
+			IsCollectionCheckedAttribute.Bind(TAttribute<ECheckBoxState>::FGetter::CreateLambda([]{ return ECheckBoxState::Unchecked; }));
+		}
 	}
 
 	TSharedPtr< STableRow< TSharedPtr<FCollectionItem> > > TableRow = SNew( STableRow< TSharedPtr<FCollectionItem> >, OwnerTable )
@@ -1379,7 +1423,6 @@ bool SCollectionView::CollectionNameChangeCommit( const TSharedPtr< FCollectionI
 	check(CollectionTreePtr->GetNumItemsSelected() == 1);
 
 	FCollectionManagerModule& CollectionManagerModule = FCollectionManagerModule::GetModule();
-	ECollectionShareType::Type CollectionType = CollectionItem->CollectionType;
 
 	// If new name is empty, set it back to the original name
 	const FName NewNameFinal( NewName.IsEmpty() ? CollectionItem->CollectionName : FName(*NewName) );
@@ -1404,7 +1447,7 @@ bool SCollectionView::CollectionNameChangeCommit( const TSharedPtr< FCollectionI
 			return false;
 		}
 
-		if ( !CollectionManagerModule.Get().CreateCollection(NewNameFinal, CollectionType) )
+		if ( !CollectionManagerModule.Get().CreateCollection(NewNameFinal, CollectionItem->CollectionType, CollectionItem->StorageMode) )
 		{
 			// Failed to add the collection, remove it from the list
 			AvailableCollections.Remove(FCollectionNameType(CollectionItem->CollectionName, CollectionItem->CollectionType));
@@ -1417,7 +1460,14 @@ bool SCollectionView::CollectionNameChangeCommit( const TSharedPtr< FCollectionI
 		if ( NewCollectionParentKey.IsSet() )
 		{
 			// Try and set the parent correctly (if this fails for any reason, the collection will still be added, but will just appear at the root)
-			CollectionManagerModule.Get().ReparentCollection(NewNameFinal, CollectionType, NewCollectionParentKey->Name, NewCollectionParentKey->Type);
+			CollectionManagerModule.Get().ReparentCollection(NewNameFinal, CollectionItem->CollectionType, NewCollectionParentKey->Name, NewCollectionParentKey->Type);
+		}
+
+		// Notify anything that cares that this collection has been created now
+		if ( CollectionItem->OnCollectionCreatedEvent.IsBound())
+		{
+			CollectionItem->OnCollectionCreatedEvent.Execute(FCollectionNameType(NewNameFinal, CollectionItem->CollectionType));
+			CollectionItem->OnCollectionCreatedEvent.Unbind();
 		}
 	}
 	else
@@ -1429,13 +1479,13 @@ bool SCollectionView::CollectionNameChangeCommit( const TSharedPtr< FCollectionI
 		}
 
 		// If the new name doesn't pass our current filter, we need to clear it
-		if ( !CollectionItemTextFilter->PassesFilter( FCollectionItem(NewNameFinal, CollectionType) ) )
+		if ( !CollectionItemTextFilter->PassesFilter( FCollectionItem(NewNameFinal, CollectionItem->CollectionType) ) )
 		{
 			SearchBoxPtr->SetText(FText::GetEmpty());
 		}
 
 		// Otherwise perform the rename
-		if ( !CollectionManagerModule.Get().RenameCollection(CollectionItem->CollectionName, CollectionType, NewNameFinal, CollectionType) )
+		if ( !CollectionManagerModule.Get().RenameCollection(CollectionItem->CollectionName, CollectionItem->CollectionType, NewNameFinal, CollectionItem->CollectionType) )
 		{
 			// Failed to rename the collection
 			OutWarningMessage = FText::Format( LOCTEXT("RenameCollectionFailed", "Failed to rename the collection. {0}"), CollectionManagerModule.Get().GetLastError());
@@ -1446,7 +1496,7 @@ bool SCollectionView::CollectionNameChangeCommit( const TSharedPtr< FCollectionI
 	// At this point CollectionItem is no longer a member of the CollectionItems list (as the list is repopulated by
 	// UpdateCollectionItems, which is called by a broadcast from CollectionManagerModule::RenameCollection, above).
 	// So search again for the item by name and type.
-	auto NewCollectionItemPtr = AvailableCollections.Find( FCollectionNameType(NewNameFinal, CollectionType) );
+	auto NewCollectionItemPtr = AvailableCollections.Find( FCollectionNameType(NewNameFinal, CollectionItem->CollectionType) );
 
 	// Reselect the path to notify that the selection has changed
 	{
diff --git a/Engine/Source/Editor/ContentBrowser/Private/SCollectionView.h b/Engine/Source/Editor/ContentBrowser/Private/SCollectionView.h
--- a/Engine/Source/Editor/ContentBrowser/Private/SCollectionView.h
+++ b/Engine/Source/Editor/ContentBrowser/Private/SCollectionView.h
@@ -3,8 +3,8 @@
 #pragma once
 
 #include "CollectionManagerModule.h"
+#include "CollectionViewTypes.h"
 
-struct FCollectionItem;
 class SCollectionListItem;
 class FCollectionAssetManagement;
 
@@ -77,8 +77,45 @@ public:
 	virtual FReply OnDragOver( const FGeometry& MyGeometry, const FDragDropEvent& DragDropEvent ) override;
 	virtual FReply OnDrop( const FGeometry& MyGeometry, const FDragDropEvent& DragDropEvent ) override;
 
+	/** Creates the menu for the save dynamic collection button */
+	void MakeSaveDynamicCollectionMenu(FText InQueryString);
+
 private:
 
+	/** Payload data for CreateCollectionItem (likely from a delegate binding) */
+	struct FCreateCollectionPayload
+	{
+		FCreateCollectionPayload()
+			: ParentCollection()
+			, OnCollectionCreatedEvent()
+		{
+		}
+
+		explicit FCreateCollectionPayload(TOptional<FCollectionNameType> InParentCollection)
+			: ParentCollection(InParentCollection)
+			, OnCollectionCreatedEvent()
+		{
+		}
+
+		explicit FCreateCollectionPayload(FCollectionItem::FCollectionCreatedEvent InCollectionCreatedEvent)
+			: ParentCollection()
+			, OnCollectionCreatedEvent(InCollectionCreatedEvent)
+		{
+		}
+
+		FCreateCollectionPayload(TOptional<FCollectionNameType> InParentCollection, FCollectionItem::FCollectionCreatedEvent InCollectionCreatedEvent)
+			: ParentCollection(InParentCollection)
+			, OnCollectionCreatedEvent(InCollectionCreatedEvent)
+		{
+		}
+
+		/** Should this collection be created as a child of another collection? */
+		TOptional<FCollectionNameType> ParentCollection;
+
+		/** Callback for after the collection has been fully created (delayed due to user naming, and potential cancellation) */
+		FCollectionItem::FCollectionCreatedEvent OnCollectionCreatedEvent;
+	};
+
 	/** True if the selection changed delegate is allowed at the moment */
 	bool ShouldAllowSelectionChangedDelegate() const;
 
@@ -95,7 +132,7 @@ private:
 	EVisibility GetAddCollectionButtonVisibility() const;
 	
 	/** Sets up an inline creation process for a new collection of the specified type */
-	void CreateCollectionItem( ECollectionShareType::Type CollectionType, TOptional<FCollectionNameType> ParentCollection );
+	void CreateCollectionItem( ECollectionShareType::Type CollectionType, ECollectionStorageMode::Type StorageMode, const FCreateCollectionPayload& InCreationPayload );
 
 	/** Sets up an inline rename for the specified collection */
 	void RenameCollectionItem( const TSharedPtr<FCollectionItem>& ItemToRename );
diff --git a/Engine/Source/Editor/ContentBrowser/Private/SContentBrowser.cpp b/Engine/Source/Editor/ContentBrowser/Private/SContentBrowser.cpp
--- a/Engine/Source/Editor/ContentBrowser/Private/SContentBrowser.cpp
+++ b/Engine/Source/Editor/ContentBrowser/Private/SContentBrowser.cpp
@@ -8,6 +8,7 @@
 #include "ContentBrowserModule.h"
 #include "ContentBrowserCommands.h"
 #include "CollectionManagerModule.h"
+#include "CollectionContextMenu.h"
 #include "AssetRegistryModule.h"
 #include "SDockTab.h"
 #include "GenericCommands.h"
@@ -542,8 +543,7 @@ void SContentBrowser::Construct( const FArguments& InArgs, const FName& InInstan
 
 						// Search
 						+SHorizontalBox::Slot()
-						.Padding(4, 0, 0, 0)
-						.VAlign(VAlign_Center)
+						.Padding(4, 1, 0, 0)
 						.FillWidth(1.0f)
 						[
 							SAssignNew(SearchBoxPtr, SAssetSearchBox)
@@ -554,6 +554,26 @@ void SContentBrowser::Construct( const FArguments& InArgs, const FName& InInstan
 							.DelayChangeNotificationsWhileTyping( true )
 							.AddMetaData<FTagMetaData>(FTagMetaData(TEXT("ContentBrowserSearchAssets")))
 						]
+
+						// Save Search
+						+SHorizontalBox::Slot()
+						.AutoWidth()
+						.VAlign(VAlign_Center)
+						.Padding(2.0f, 0.0f, 0.0f, 0.0f)
+						[
+							SNew(SButton)
+							.ButtonStyle(FEditorStyle::Get(), "FlatButton")
+							.ToolTipText(LOCTEXT("SaveSearchButtonTooltip", "Save the current search as a dynamic collection."))
+							.IsEnabled(this, &SContentBrowser::IsSaveSearchButtonEnabled)
+							.OnClicked(this, &SContentBrowser::OnSaveSearchButtonClicked)
+							.ContentPadding( FMargin(1, 1) )
+							[
+								SNew(STextBlock)
+								.TextStyle(FEditorStyle::Get(), "ContentBrowser.Filters.Text")
+								.Font(FEditorStyle::Get().GetFontStyle("FontAwesome.10"))
+								.Text(FEditorFontGlyphs::Floppy_O)
+							]
+						]
 					]
 
 					// Filters
@@ -608,9 +628,8 @@ void SContentBrowser::Construct( const FArguments& InArgs, const FName& InInstan
 
 
 	// Select /Game by default
-	FSourcesData DefaultSourcesData;
+	FSourcesData DefaultSourcesData(FName("/Game"));
 	TArray<FString> SelectedPaths;
-	DefaultSourcesData.PackagePaths.Add(TEXT("/Game"));
 	SelectedPaths.Add(TEXT("/Game"));
 	PathViewPtr->SetSelectedPaths(SelectedPaths);
 	AssetViewPtr->SetSourcesData(DefaultSourcesData);
@@ -1022,12 +1041,45 @@ void SContentBrowser::SourcesChanged(const TArray<FString>& SelectedPaths, const
 	UE_LOG(LogContentBrowser, VeryVerbose, TEXT("The content browser source was changed by the sources view to '%s'"), *NewSource);
 
 	FSourcesData SourcesData;
-	for (int32 PathIdx = 0; PathIdx < SelectedPaths.Num(); ++PathIdx)
 	{
-		SourcesData.PackagePaths.Add(FName(*SelectedPaths[PathIdx]));
+		TArray<FName> SelectedPathNames;
+		SelectedPathNames.Reserve(SelectedPaths.Num());
+		for (const FString& SelectedPath : SelectedPaths)
+		{
+			SelectedPathNames.Add(FName(*SelectedPath));
+		}
+		SourcesData = FSourcesData(MoveTemp(SelectedPathNames), SelectedCollections);
+	}
+
+	// A dynamic collection should apply its search query to the CB search, so we need to stash the current search so that we can restore it again later
+	if (SourcesData.Collections.Num() == 1 &&  SourcesData.IsDynamicCollection())
+	{
+		// Only stash the user search term once in case we're switching between dynamic collections
+		if (!StashedSearchBoxText.IsSet())
+		{
+			StashedSearchBoxText = TextFilter->GetRawFilterText();
+		}
+
+		FCollectionManagerModule& CollectionManagerModule = FCollectionManagerModule::GetModule();
+
+		const FCollectionNameType& DynamicCollection = SourcesData.Collections[0];
+
+		FString DynamicQueryString;
+		CollectionManagerModule.Get().GetDynamicQueryText(DynamicCollection.Name, DynamicCollection.Type, DynamicQueryString);
+
+		const FText DynamicQueryText = FText::FromString(DynamicQueryString);
+		SetSearchBoxText(DynamicQueryText);
+		SearchBoxPtr->SetText(DynamicQueryText);
+	}
+	else if (StashedSearchBoxText.IsSet())
+	{
+		// Restore the stashed search term
+		const FText StashedText = StashedSearchBoxText.GetValue();
+		StashedSearchBoxText.Reset();
+
+		SetSearchBoxText(StashedText);
+		SearchBoxPtr->SetText(StashedText);
 	}
-	
-	SourcesData.Collections = SelectedCollections;
 
 	if (!AssetViewPtr->GetSourcesData().IsEmpty())
 	{
@@ -1158,7 +1210,7 @@ void SContentBrowser::OnUpdateHistoryData(FHistoryData& HistoryData) const
 	const FSourcesData& SourcesData = AssetViewPtr->GetSourcesData();
 	const TArray<FAssetData>& SelectedAssets = AssetViewPtr->GetSelectedAssets();
 
-	const FText NewSource = SourcesData.PackagePaths.Num() > 0 ? FText::FromName(SourcesData.PackagePaths[0]) : (SourcesData.Collections.Num() > 0 ? FText::FromName(SourcesData.Collections[0].Name) : LOCTEXT("AllAssets", "All Assets"));
+	const FText NewSource = SourcesData.HasPackagePaths() ? FText::FromName(SourcesData.PackagePaths[0]) : (SourcesData.HasCollections() ? FText::FromName(SourcesData.Collections[0].Name) : LOCTEXT("AllAssets", "All Assets"));
 
 	HistoryData.HistoryDesc = NewSource;
 	HistoryData.SourcesData = SourcesData;
@@ -1248,11 +1300,58 @@ void SContentBrowser::OnSearchBoxCommitted(const FText& InSearchText, ETextCommi
 	SetSearchBoxText(InSearchText);
 }
 
+bool SContentBrowser::IsSaveSearchButtonEnabled() const
+{
+	return !TextFilter->GetRawFilterText().IsEmptyOrWhitespace();
+}
+
+FReply SContentBrowser::OnSaveSearchButtonClicked()
+{
+	// Need to make sure we can see the collections view
+	if (!bSourcesViewExpanded)
+	{
+		SourcesViewExpandClicked();
+	}
+
+	// We want to add any currently selected paths to the final saved query so that you get back roughly the same list of objects as what you're currently seeing
+	FString SelectedPathsQuery;
+	{
+		auto SelectedPaths = PathViewPtr->GetSelectedPaths();
+		for (int32 SelectedPathIndex = 0; SelectedPathIndex < SelectedPaths.Num(); ++SelectedPathIndex)
+		{
+			SelectedPathsQuery.Append(TEXT("Path:'"));
+			SelectedPathsQuery.Append(SelectedPaths[SelectedPathIndex]);
+			SelectedPathsQuery.Append(TEXT("'..."));
+
+			if (SelectedPathIndex + 1 < SelectedPaths.Num())
+			{
+				SelectedPathsQuery.Append(TEXT(" OR "));
+			}
+		}
+	}
+
+	// todo: should we automatically append any type filters too?
+
+	// Produce the final query
+	FText FinalQueryText;
+	if (SelectedPathsQuery.IsEmpty())
+	{
+		FinalQueryText = TextFilter->GetRawFilterText();
+	}
+	else
+	{
+		FinalQueryText = FText::FromString(FString::Printf(TEXT("(%s) AND (%s)"), *TextFilter->GetRawFilterText().ToString(), *SelectedPathsQuery));
+	}
+
+	CollectionViewPtr->MakeSaveDynamicCollectionMenu(FinalQueryText);
+	return FReply::Handled();
+}
+
 void SContentBrowser::OnPathClicked( const FString& CrumbData )
 {
 	FSourcesData SourcesData = AssetViewPtr->GetSourcesData();
 
-	if ( SourcesData.Collections.Num() > 0 )
+	if ( SourcesData.HasCollections() )
 	{
 		// Collection crumb was clicked. See if we've clicked on a different collection in the hierarchy, and change the path if required.
 		TOptional<FCollectionNameType> CollectionClicked;
@@ -1278,7 +1377,7 @@ void SContentBrowser::OnPathClicked( const FString& CrumbData )
 			CollectionSelected(CollectionClicked.GetValue());
 		}
 	}
-	else if ( SourcesData.PackagePaths.Num() == 0 )
+	else if ( !SourcesData.HasPackagePaths() )
 	{
 		// No collections or paths are selected. This is "All Assets". Don't change the path when this is clicked.
 	}
@@ -1290,8 +1389,6 @@ void SContentBrowser::OnPathClicked( const FString& CrumbData )
 		PathViewPtr->SetSelectedPaths(SelectedPaths);
 
 		PathSelected(CrumbData);
-		
-		SourcesChanged(SelectedPaths, TArray<FCollectionNameType>());
 	}
 }
 
@@ -1307,7 +1404,7 @@ TSharedPtr<SWidget> SContentBrowser::OnGetCrumbDelimiterContent(const FString& C
 	TSharedPtr<SWidget> Widget = SNullWidget::NullWidget;
 	TSharedPtr<SWidget> MenuWidget;
 
-	if( SourcesData.Collections.Num() > 0 )
+	if( SourcesData.HasCollections() )
 	{
 		TOptional<FCollectionNameType> CollectionClicked;
 		{
@@ -1350,7 +1447,7 @@ TSharedPtr<SWidget> SContentBrowser::OnGetCrumbDelimiterContent(const FString& C
 			}
 		}
 	}
-	else if( SourcesData.PackagePaths.Num() > 0 )
+	else if( SourcesData.HasPackagePaths() )
 	{
 		TArray<FString> SubPaths;
 		const bool bRecurse = false;
@@ -1414,7 +1511,7 @@ TSharedRef<SWidget> SContentBrowser::GetPathPickerContent()
 	FPathPickerConfig PathPickerConfig;
 
 	FSourcesData SourcesData = AssetViewPtr->GetSourcesData();
-	if ( SourcesData.PackagePaths.Num() > 0 )
+	if ( SourcesData.HasPackagePaths() )
 	{
 		PathPickerConfig.DefaultPath = SourcesData.PackagePaths[0].ToString();
 	}
@@ -1454,7 +1551,7 @@ FString SContentBrowser::GetCurrentPath() const
 {
 	FString CurrentPath;
 	const FSourcesData& SourcesData = AssetViewPtr->GetSourcesData();
-	if ( SourcesData.PackagePaths.Num() > 0 && SourcesData.PackagePaths[0] != NAME_None )
+	if ( SourcesData.HasPackagePaths() && SourcesData.PackagePaths[0] != NAME_None )
 	{
 		CurrentPath = SourcesData.PackagePaths[0].ToString();
 	}
@@ -1725,7 +1822,7 @@ FReply SContentBrowser::SourcesViewExpandClicked()
 {
 	bSourcesViewExpanded = !bSourcesViewExpanded;
 
-	// Notify 'Soureces View Expanded' delegate
+	// Notify 'Sources View Expanded' delegate
 	FContentBrowserModule& ContentBrowserModule = FModuleManager::GetModuleChecked<FContentBrowserModule>( TEXT("ContentBrowser") );
 	FContentBrowserModule::FOnSourcesViewChanged& SourcesViewChangedDelegate = ContentBrowserModule.GetOnSourcesViewChanged();
 	if(SourcesViewChangedDelegate.IsBound())
@@ -1966,7 +2063,7 @@ void SContentBrowser::UpdatePath()
 
 	PathBreadcrumbTrail->ClearCrumbs();
 
-	if ( SourcesData.PackagePaths.Num() > 0 )
+	if ( SourcesData.HasPackagePaths() )
 	{
 		TArray<FString> Crumbs;
 		SourcesData.PackagePaths[0].ToString().ParseIntoArray(Crumbs, TEXT("/"), true);
@@ -1982,7 +2079,7 @@ void SContentBrowser::UpdatePath()
 			CrumbPath += TEXT("/");
 		}
 	}
-	else if ( SourcesData.Collections.Num() > 0 )
+	else if ( SourcesData.HasCollections() )
 	{
 		FCollectionManagerModule& CollectionManagerModule = FCollectionManagerModule::GetModule();
 		TArray<FCollectionNameType> CollectionPathItems;
@@ -2033,11 +2130,11 @@ FString SContentBrowser::GetPathText() const
 		const FSourcesData& SourcesData = AssetViewPtr->GetSourcesData();
 
 		// At least one source is selected
-		int32 NumSources = SourcesData.PackagePaths.Num() + SourcesData.Collections.Num();
+		const int32 NumSources = SourcesData.PackagePaths.Num() + SourcesData.Collections.Num();
 
 		if (NumSources > 0)
 		{
-			PathLabelText = SourcesData.PackagePaths.Num() > 0 ? SourcesData.PackagePaths[0].ToString() : SourcesData.Collections[0].Name.ToString();
+			PathLabelText = SourcesData.HasPackagePaths() ? SourcesData.PackagePaths[0].ToString() : SourcesData.Collections[0].Name.ToString();
 
 			if (NumSources > 1)
 			{
@@ -2056,7 +2153,7 @@ FString SContentBrowser::GetPathText() const
 bool SContentBrowser::IsFilteredBySource() const
 {
 	const FSourcesData& SourcesData = AssetViewPtr->GetSourcesData();
-	return SourcesData.PackagePaths.Num() != 0 || SourcesData.Collections.Num() != 0;
+	return !SourcesData.IsEmpty();
 }
 
 void SContentBrowser::OnAssetRenameCommitted(const TArray<FAssetData>& Assets)
@@ -2103,8 +2200,7 @@ void SContentBrowser::OnOpenedFolderDeleted()
 	DefaultSelectedPaths.Add(TEXT("/Game"));
 	PathViewPtr->SetSelectedPaths(DefaultSelectedPaths);
 
-	FSourcesData DefaultSourcesData;
-	DefaultSourcesData.PackagePaths.Add(TEXT("/Game"));
+	FSourcesData DefaultSourcesData(FName("/Game"));
 	AssetViewPtr->SetSourcesData(DefaultSourcesData);
 }
 
diff --git a/Engine/Source/Editor/ContentBrowser/Private/SContentBrowser.h b/Engine/Source/Editor/ContentBrowser/Private/SContentBrowser.h
--- a/Engine/Source/Editor/ContentBrowser/Private/SContentBrowser.h
+++ b/Engine/Source/Editor/ContentBrowser/Private/SContentBrowser.h
@@ -138,6 +138,12 @@ private:
 	/** Called by the editable text control when the user commits a text change */
 	void OnSearchBoxCommitted(const FText& InSearchText, ETextCommit::Type CommitInfo);
 
+	/** Should the "Save Search" button be enabled? */
+	bool IsSaveSearchButtonEnabled() const;
+
+	/** Open the menu to let you save the current search text as a dynamic collection */
+	FReply OnSaveSearchButtonClicked();
+
 	/** Called when a crumb in the path breadcrumb trail or menu is clicked */
 	void OnPathClicked(const FString& CrumbData);
 
@@ -387,6 +393,9 @@ private:
 	/** The splitter between the path & collection view */
 	TSharedPtr<SSplitter> PathCollectionSplitterPtr;
 
+	/** When viewing a dynamic collection, the active search query will be stashed in this variable so that it can be restored again later */
+	TOptional<FText> StashedSearchBoxText;
+
 public: 
 
 	/** The section of EditorPerProjectUserSettings in which to save content browser settings */
diff --git a/Engine/Source/Editor/ContentBrowser/Private/SPathView.cpp b/Engine/Source/Editor/ContentBrowser/Private/SPathView.cpp
--- a/Engine/Source/Editor/ContentBrowser/Private/SPathView.cpp
+++ b/Engine/Source/Editor/ContentBrowser/Private/SPathView.cpp
@@ -553,9 +553,9 @@ void SPathView::ApplyHistoryData ( const FHistoryData& History )
 
 	// Update paths
 	TArray<FString> SelectedPaths;
-	for ( auto PathIt = History.SourcesData.PackagePaths.CreateConstIterator(); PathIt; ++PathIt)
+	for (const FName& HistoryPath : History.SourcesData.PackagePaths)
 	{
-		SelectedPaths.Add((*PathIt).ToString());
+		SelectedPaths.Add(HistoryPath.ToString());
 	}
 	SetSelectedPaths(SelectedPaths);
 }
diff --git a/Engine/Source/Editor/ContentBrowser/Private/SourcesData.h b/Engine/Source/Editor/ContentBrowser/Private/SourcesData.h
--- a/Engine/Source/Editor/ContentBrowser/Private/SourcesData.h
+++ b/Engine/Source/Editor/ContentBrowser/Private/SourcesData.h
@@ -7,11 +7,93 @@ struct FSourcesData
 	TArray<FName> PackagePaths;
 	TArray<FCollectionNameType> Collections;
 
-	bool IsEmpty() const
+	FSourcesData()
+		: PackagePaths()
+		, Collections()
+	{
+	}
+
+	explicit FSourcesData(FName InPackagePath)
+		: PackagePaths()
+		, Collections()
+	{
+		PackagePaths.Add(InPackagePath);
+	}
+
+	explicit FSourcesData(FCollectionNameType InCollection)
+		: PackagePaths()
+		, Collections()
+	{
+		Collections.Add(InCollection);
+	}
+
+	FSourcesData(TArray<FName> InPackagePaths, TArray<FCollectionNameType> InCollections)
+		: PackagePaths(MoveTemp(InPackagePaths))
+		, Collections(MoveTemp(InCollections))
+	{
+	}
+
+	FSourcesData(const FSourcesData& Other)
+		: PackagePaths(Other.PackagePaths)
+		, Collections(Other.Collections)
+	{
+	}
+
+	FSourcesData(FSourcesData&& Other)
+		: PackagePaths(MoveTemp(Other.PackagePaths))
+		, Collections(MoveTemp(Other.Collections))
+	{
+	}
+
+	FSourcesData& operator=(const FSourcesData& Other)
+	{
+		if (this != &Other)
+		{
+			PackagePaths = Other.PackagePaths;
+			Collections = Other.Collections;
+		}
+		return *this;
+	}
+
+	FSourcesData& operator=(FSourcesData&& Other)
+	{
+		if (this != &Other)
+		{
+			PackagePaths = MoveTemp(Other.PackagePaths);
+			Collections = MoveTemp(Other.Collections);
+		}
+		return *this;
+	}
+
+	FORCEINLINE bool IsEmpty() const
 	{
 		return PackagePaths.Num() == 0 && Collections.Num() == 0;
 	}
 
+	FORCEINLINE bool HasPackagePaths() const
+	{
+		return PackagePaths.Num() > 0;
+	}
+
+	FORCEINLINE bool HasCollections() const
+	{
+		return Collections.Num() > 0;
+	}
+
+	bool IsDynamicCollection() const
+	{
+		if ( Collections.Num() == 1 && FCollectionManagerModule::IsModuleAvailable() )
+		{
+			// Collection manager module should already be loaded since it may cause a hitch on the first search
+			FCollectionManagerModule& CollectionManagerModule = FCollectionManagerModule::GetModule();
+
+			ECollectionStorageMode::Type StorageMode = ECollectionStorageMode::Static;
+			return (CollectionManagerModule.Get().GetCollectionStorageMode(Collections[0].Name, Collections[0].Type, StorageMode) && StorageMode == ECollectionStorageMode::Dynamic);
+		}
+
+		return false;
+	}
+
 	FARFilter MakeFilter(bool bRecurse, bool bUsingFolders) const
 	{
 		FARFilter Filter;
@@ -20,6 +102,12 @@ struct FSourcesData
 		Filter.PackagePaths = PackagePaths;
 		Filter.bRecursivePaths = bRecurse || !bUsingFolders;
 
+		// If we have a dynamic source, then we need to make sure that the root path is searched for matching objects as the dynamic filter will sort through them
+		if (IsDynamicCollection())
+		{
+			Filter.PackagePaths.AddUnique(TEXT("/"));
+		}
+
 		// Collections
 		TArray<FName> ObjectPathsFromCollections;
 		if ( Collections.Num() && FCollectionManagerModule::IsModuleAvailable() )
@@ -41,4 +129,4 @@ struct FSourcesData
 
 		return Filter;
 	}
-};
\ No newline at end of file
+};
diff --git a/Engine/Source/Editor/ContentBrowser/Private/SourcesViewWidgets.cpp b/Engine/Source/Editor/ContentBrowser/Private/SourcesViewWidgets.cpp
--- a/Engine/Source/Editor/ContentBrowser/Private/SourcesViewWidgets.cpp
+++ b/Engine/Source/Editor/ContentBrowser/Private/SourcesViewWidgets.cpp
@@ -329,7 +329,6 @@ void SCollectionTreeItem::Construct( const FArguments& InArgs )
 		.Padding(0)
 		[
 			SNew(SHorizontalBox)
-			.ToolTip( SNew(SToolTip).Text( ECollectionShareType::GetDescription(InArgs._CollectionItem->CollectionType) ) )
 
 			+SHorizontalBox::Slot()
 			.AutoWidth()
@@ -348,10 +347,11 @@ void SCollectionTreeItem::Construct( const FArguments& InArgs )
 			.VAlign(VAlign_Center)
 			.Padding(0, 0, 2, 0)
 			[
-				// Type Icon
+				// Share Type Icon
 				SNew(SImage)
 				.Image( FEditorStyle::GetBrush( ECollectionShareType::GetIconStyleName(InArgs._CollectionItem->CollectionType) ) )
 				.ColorAndOpacity( this, &SCollectionTreeItem::GetCollectionColor )
+				.ToolTipText(ECollectionShareType::GetDescription(InArgs._CollectionItem->CollectionType))
 			]
 
 			+SHorizontalBox::Slot()
@@ -367,6 +367,25 @@ void SCollectionTreeItem::Construct( const FArguments& InArgs )
 				.IsReadOnly( InArgs._IsReadOnly )
 			]
 
+			+SHorizontalBox::Slot()
+			.AutoWidth()
+			.Padding(2, 0, 3, 0)
+			[
+				// Storage Mode Icon
+				SNew(SBox)
+				.WidthOverride(16)
+				.HeightOverride(16)
+				.VAlign(VAlign_Center)
+				.HAlign(HAlign_Center)
+				[
+					SNew(STextBlock)
+					.Font(FEditorStyle::Get().GetFontStyle("FontAwesome.10"))
+					.Text(this, &SCollectionTreeItem::GetCollectionStorageModeIconText)
+					.ColorAndOpacity(FLinearColor::Gray)
+					.ToolTipText(this, &SCollectionTreeItem::GetCollectionStorageModeToolTipText)
+				]
+			]
+
 			+SHorizontalBox::Slot()
 			.AutoWidth()
 			.VAlign(VAlign_Center)
@@ -545,6 +564,40 @@ const FSlateBrush* SCollectionTreeItem::GetBorderImage() const
 	return bDraggedOver ? FEditorStyle::GetBrush("Menu.Background") : FEditorStyle::GetBrush("NoBorder");
 }
 
+FText SCollectionTreeItem::GetCollectionStorageModeIconText() const
+{
+	TSharedPtr<FCollectionItem> CollectionItemPtr = CollectionItem.Pin();
+
+	if (CollectionItemPtr.IsValid())
+	{
+		switch(CollectionItemPtr->StorageMode)
+		{
+		case ECollectionStorageMode::Static:
+			return FEditorFontGlyphs::List_Alt;
+
+		case ECollectionStorageMode::Dynamic:
+			return FEditorFontGlyphs::Bolt;
+
+		default:
+			break;
+		}
+	}
+
+	return FText::GetEmpty();
+}
+
+FText SCollectionTreeItem::GetCollectionStorageModeToolTipText() const
+{
+	TSharedPtr<FCollectionItem> CollectionItemPtr = CollectionItem.Pin();
+
+	if (CollectionItemPtr.IsValid())
+	{
+		return ECollectionStorageMode::GetDescription(CollectionItemPtr->StorageMode);
+	}
+
+	return FText::GetEmpty();
+}
+
 FSlateColor SCollectionTreeItem::GetCollectionStatusColor() const
 {
 	TSharedPtr<FCollectionItem> CollectionItemPtr = CollectionItem.Pin();
diff --git a/Engine/Source/Editor/ContentBrowser/Private/SourcesViewWidgets.h b/Engine/Source/Editor/ContentBrowser/Private/SourcesViewWidgets.h
--- a/Engine/Source/Editor/ContentBrowser/Private/SourcesViewWidgets.h
+++ b/Engine/Source/Editor/ContentBrowser/Private/SourcesViewWidgets.h
@@ -240,10 +240,16 @@ private:
 	/** Returns the image for the border around this item. Used for drag/drop operations */
 	const FSlateBrush* GetBorderImage() const;
 
+	/** Get the FontAwesome icon corresponding to the current collection storage mode */
+	FText GetCollectionStorageModeIconText() const;
+
+	/** Get the tooltip corresponding to the current collection storage mode */
+	FText GetCollectionStorageModeToolTipText() const;
+
 	/** Get the color to use for the collection item status */
 	FSlateColor GetCollectionStatusColor() const;
 
-	/** Get the tooltip correspinding to the current collection item status */
+	/** Get the tooltip corresponding to the current collection item status */
 	FText GetCollectionStatusToolTipText() const;
 
 private:
diff --git a/Engine/Source/Editor/UnrealEd/Private/Commandlets/PackageUtilities.cpp b/Engine/Source/Editor/UnrealEd/Private/Commandlets/PackageUtilities.cpp
--- a/Engine/Source/Editor/UnrealEd/Private/Commandlets/PackageUtilities.cpp
+++ b/Engine/Source/Editor/UnrealEd/Private/Commandlets/PackageUtilities.cpp
@@ -235,7 +235,7 @@ public:
 	static bool CreateAssetSet(FName InSetName, ECollectionShareType::Type InSetType)
 	{
 		FCollectionManagerModule& CollectionManagerModule = FCollectionManagerModule::GetModule();
-		return CollectionManagerModule.Get().CreateCollection(InSetName, InSetType);
+		return CollectionManagerModule.Get().CreateCollection(InSetName, InSetType, ECollectionStorageMode::Static);
 	}
 
 	static bool DestroyAssetSet(FName InSetName, ECollectionShareType::Type InSetType )
diff --git a/Engine/Source/Runtime/AssetRegistry/Private/AssetRegistry.cpp b/Engine/Source/Runtime/AssetRegistry/Private/AssetRegistry.cpp
--- a/Engine/Source/Runtime/AssetRegistry/Private/AssetRegistry.cpp
+++ b/Engine/Source/Runtime/AssetRegistry/Private/AssetRegistry.cpp
@@ -886,7 +886,7 @@ void FAssetRegistry::RunAssetsThroughFilter(TArray<FAssetData>& AssetDataList, c
 						if ( AssetPackagePath.StartsWith(Path) )
 						{
 							// Only match the exact path or a path that starts with the target path followed by a slash
-							if ( Path.Len() == AssetPackagePath.Len() || AssetPackagePath.Mid(Path.Len(), 1) == TEXT("/") )
+							if ( Path.Len() == 1 || Path.Len() == AssetPackagePath.Len() || AssetPackagePath.Mid(Path.Len(), 1) == TEXT("/") )
 							{
 								bPassesPackagePaths = true;
 								break;
