Fixed the collection manager watching non-collection directories for changes

UE-18771 - Collections log spam when building engine dlls that live in the same directory as a running instance of the editor

The collection manager now ensures that the collection directories exist on disk before it starts watching them, and will also ignore the "system" collection directory (as that is blank and was causing it to watch for changes in the Binaries folder).

As a belt-and-braces measure, FCollectionManager::TickFileCache will also now ensure that a file ends with the ".collection" extension before attempting to process the file as a collection.

diff --git a/Engine/Source/Developer/CollectionManager/Private/CollectionManager.cpp b/Engine/Source/Developer/CollectionManager/Private/CollectionManager.cpp
--- a/Engine/Source/Developer/CollectionManager/Private/CollectionManager.cpp
+++ b/Engine/Source/Developer/CollectionManager/Private/CollectionManager.cpp
@@ -236,6 +236,14 @@ FCollectionManager::FCollectionManager()
 	{
 		const FString& CollectionFolder = CollectionFolders[CacheIdx];
 
+		if (CollectionFolder.IsEmpty())
+		{
+			continue;
+		}
+
+		// Make sure the folder we want to watch exists on disk
+		IFileManager::Get().MakeDirectory(*CollectionFolder, true);
+
 		DirectoryWatcher::FFileCacheConfig FileCacheConfig(FPaths::ConvertRelativePathToFull(CollectionFolder), FString());
 		FileCacheConfig.DetectMoves(false);
 		FileCacheConfig.RequireFileHashes(false);
@@ -1578,7 +1586,13 @@ bool FCollectionManager::TickFileCache(float InDeltaTime)
 			const auto FileCacheChanges = FileCache->GetOutstandingChanges();
 			for (const auto& FileCacheChange : FileCacheChanges)
 			{
-				const FName CollectionName = *FPaths::GetBaseFilename(FileCacheChange.Filename.Get());
+				const FString CollectionFilename = FileCacheChange.Filename.Get();
+				if (FPaths::GetExtension(CollectionFilename) != CollectionExtension)
+				{
+					continue;
+				}
+
+				const FName CollectionName = *FPaths::GetBaseFilename(CollectionFilename);
 
 				ECollectionFileAction CollectionFileAction = ECollectionFileAction::None;
 				switch (FileCacheChange.Action)
