Added control over the collection paths in the asset tooltips.
Also Added wrapping for cases where an asset is in many collections.

diff --git a/Engine/Source/Developer/CollectionManager/Private/CollectionManager.cpp b/Engine/Source/Developer/CollectionManager/Private/CollectionManager.cpp
--- a/Engine/Source/Developer/CollectionManager/Private/CollectionManager.cpp
+++ b/Engine/Source/Developer/CollectionManager/Private/CollectionManager.cpp
@@ -598,7 +598,7 @@ void FCollectionManager::GetCollectionsContainingObjects(const TArray<FName>& Ob
 	}
 }
 
-FString FCollectionManager::GetCollectionsStringForObject(FName ObjectPath, ECollectionShareType::Type ShareType, ECollectionRecursionFlags::Flags RecursionMode) const
+FString FCollectionManager::GetCollectionsStringForObject(FName ObjectPath, ECollectionShareType::Type ShareType, ECollectionRecursionFlags::Flags RecursionMode, bool bFullPaths) const
 {
 	const FCollectionObjectsMap& CachedObjects = CollectionCache.GetCachedObjects();
 
@@ -619,10 +619,16 @@ FString FCollectionManager::GetCollectionsStringForObject(FName ObjectPath, ECol
 		{
 			if ((ShareType == ECollectionShareType::CST_All || ShareType == ObjectCollectionInfo.CollectionKey.Type) && (RecursionMode & ObjectCollectionInfo.Reason) != 0)
 			{
-				CollectionPathStrings.Reset();
-				CollectionCache.RecursionHelper_DoWork(ObjectCollectionInfo.CollectionKey, ECollectionRecursionFlags::SelfAndParents, GetCollectionsStringForObjectWorker);
-
-				CollectionNameStrings.Add(FString::Join(CollectionPathStrings, TEXT("/")));
+				if (bFullPaths)
+				{
+					CollectionPathStrings.Reset();
+					CollectionCache.RecursionHelper_DoWork(ObjectCollectionInfo.CollectionKey, ECollectionRecursionFlags::SelfAndParents, GetCollectionsStringForObjectWorker);
+					CollectionNameStrings.Add(FString::Join(CollectionPathStrings, TEXT("/")));
+				}
+				else
+				{
+					CollectionNameStrings.Add(ObjectCollectionInfo.CollectionKey.Name.ToString());
+				}
 			}
 		}
 
diff --git a/Engine/Source/Developer/CollectionManager/Private/CollectionManager.h b/Engine/Source/Developer/CollectionManager/Private/CollectionManager.h
--- a/Engine/Source/Developer/CollectionManager/Private/CollectionManager.h
+++ b/Engine/Source/Developer/CollectionManager/Private/CollectionManager.h
@@ -119,7 +119,7 @@ public:
 	virtual void GetCollectionsContainingObject(FName ObjectPath, ECollectionShareType::Type ShareType, TArray<FName>& OutCollectionNames, ECollectionRecursionFlags::Flags RecursionMode = ECollectionRecursionFlags::Self) const override;
 	virtual void GetCollectionsContainingObject(FName ObjectPath, TArray<FCollectionNameType>& OutCollections, ECollectionRecursionFlags::Flags RecursionMode = ECollectionRecursionFlags::Self) const override;
 	virtual void GetCollectionsContainingObjects(const TArray<FName>& ObjectPaths, TMap<FCollectionNameType, TArray<FName>>& OutCollectionsAndMatchedObjects, ECollectionRecursionFlags::Flags RecursionMode = ECollectionRecursionFlags::Self) const override;
-	virtual FString GetCollectionsStringForObject(FName ObjectPath, ECollectionShareType::Type ShareType, ECollectionRecursionFlags::Flags RecursionMode = ECollectionRecursionFlags::Self) const override;
+	virtual FString GetCollectionsStringForObject(FName ObjectPath, ECollectionShareType::Type ShareType, ECollectionRecursionFlags::Flags RecursionMode = ECollectionRecursionFlags::Self, bool bFullPaths = true) const override;
 	virtual void CreateUniqueCollectionName(const FName& BaseName, ECollectionShareType::Type ShareType, FName& OutCollectionName) const override;
 	virtual bool IsValidCollectionName(const FString& CollectionName, ECollectionShareType::Type ShareType) const override;
 	virtual bool CreateCollection(FName CollectionName, ECollectionShareType::Type ShareType, ECollectionStorageMode::Type StorageMode) override;
diff --git a/Engine/Source/Developer/CollectionManager/Public/ICollectionManager.h b/Engine/Source/Developer/CollectionManager/Public/ICollectionManager.h
--- a/Engine/Source/Developer/CollectionManager/Public/ICollectionManager.h
+++ b/Engine/Source/Developer/CollectionManager/Public/ICollectionManager.h
@@ -59,7 +59,7 @@ public:
 	virtual void GetCollectionsContainingObjects(const TArray<FName>& ObjectPaths, TMap<FCollectionNameType, TArray<FName>>& OutCollectionsAndMatchedObjects, ECollectionRecursionFlags::Flags RecursionMode = ECollectionRecursionFlags::Self) const = 0;
 
 	/** Returns a string containing a comma separated list of collections in which the specified object exists of the specified share type */
-	virtual FString GetCollectionsStringForObject(FName ObjectPath, ECollectionShareType::Type ShareType, ECollectionRecursionFlags::Flags RecursionMode = ECollectionRecursionFlags::Self) const = 0;
+	virtual FString GetCollectionsStringForObject(FName ObjectPath, ECollectionShareType::Type ShareType, ECollectionRecursionFlags::Flags RecursionMode = ECollectionRecursionFlags::Self, bool bFullPaths = true) const = 0;
 
 	/** Creates a unique collection name for the given type taking the form BaseName+(unique number) */
 	virtual void CreateUniqueCollectionName(const FName& BaseName, ECollectionShareType::Type ShareType, FName& OutCollectionName) const = 0;
diff --git a/Engine/Source/Editor/ContentBrowser/Private/AssetViewWidgets.cpp b/Engine/Source/Editor/ContentBrowser/Private/AssetViewWidgets.cpp
--- a/Engine/Source/Editor/ContentBrowser/Private/AssetViewWidgets.cpp
+++ b/Engine/Source/Editor/ContentBrowser/Private/AssetViewWidgets.cpp
@@ -33,6 +33,7 @@
 #include "ContentBrowserLog.h"
 #include "ObjectTools.h"
 #include "AssetThumbnail.h"
+#include "Settings/ContentBrowserSettings.h"
 
 #define LOCTEXT_NAMESPACE "ContentBrowser"
 
@@ -703,8 +704,7 @@ TSharedRef<SWidget> SAssetViewItem::CreateToolTipWidget() const
 			// Add Collections
 			{
 				FCollectionManagerModule& CollectionManagerModule = FCollectionManagerModule::GetModule();
-
-				const FString CollectionNames = CollectionManagerModule.Get().GetCollectionsStringForObject(AssetData.ObjectPath, ECollectionShareType::CST_All);
+				const FString CollectionNames = CollectionManagerModule.Get().GetCollectionsStringForObject(AssetData.ObjectPath, ECollectionShareType::CST_All, ECollectionRecursionFlags::Self, GetDefault<UContentBrowserSettings>()->bShowFullCollectionNameInToolTip);
 				if (!CollectionNames.IsEmpty())
 				{
 					AddToToolTipInfoBox(InfoBox, LOCTEXT("AssetToolTipKey_Collections", "Collections"), FText::FromString(CollectionNames), false);
@@ -951,6 +951,7 @@ void SAssetViewItem::AddToToolTipInfoBox(const TSharedRef<SVerticalBox>& InfoBox
 			SNew(STextBlock) .Text( Value )
 			.ColorAndOpacity(bImportant ? ImportantStyle.GetForegroundColor() : FSlateColor::UseForeground())
 			.HighlightText((Key.ToString() == TEXT("Path")) ? HighlightText : FText())
+			.WrapTextAt(700.0f)
 		]
 	];
 }
diff --git a/Engine/Source/Editor/UnrealEd/Classes/Settings/ContentBrowserSettings.h b/Engine/Source/Editor/UnrealEd/Classes/Settings/ContentBrowserSettings.h
--- a/Engine/Source/Editor/UnrealEd/Classes/Settings/ContentBrowserSettings.h
+++ b/Engine/Source/Editor/UnrealEd/Classes/Settings/ContentBrowserSettings.h
@@ -46,6 +46,10 @@ public:
 	UPROPERTY(EditAnywhere, config, Category = ContentBrowser, meta = (DisplayName = "Number of Assets to Keep in the Recently Opened Filter", ClampMin = "1", ClampMax = "30"))
 	int32 NumObjectsInRecentList;
 
+	/** Whether the Content Browser should open the Sources Panel by default */
+	UPROPERTY(EditAnywhere, config, Category = ContentBrowser)
+	bool bShowFullCollectionNameInToolTip;
+
 public:
 
 	/** Sets whether we are allowed to display the engine folder or not, optional flag for setting override instead */
diff --git a/Engine/Source/Editor/UnrealEd/Private/Settings/SettingsClasses.cpp b/Engine/Source/Editor/UnrealEd/Private/Settings/SettingsClasses.cpp
--- a/Engine/Source/Editor/UnrealEd/Private/Settings/SettingsClasses.cpp
+++ b/Engine/Source/Editor/UnrealEd/Private/Settings/SettingsClasses.cpp
@@ -45,6 +45,7 @@ UContentBrowserSettings::FSettingChangedEvent UContentBrowserSettings::SettingCh
 
 UContentBrowserSettings::UContentBrowserSettings( const FObjectInitializer& ObjectInitializer )
 	: Super(ObjectInitializer)
+	, bShowFullCollectionNameInToolTip(true)
 {
 }
 
