Fixed implicit conversion warning in VS2019.

diff --git a/Engine/Source/Developer/CollectionManager/Private/CollectionManager.cpp b/Engine/Source/Developer/CollectionManager/Private/CollectionManager.cpp
--- a/Engine/Source/Developer/CollectionManager/Private/CollectionManager.cpp
+++ b/Engine/Source/Developer/CollectionManager/Private/CollectionManager.cpp
@@ -174,7 +174,7 @@ const FCollectionColorArray& FCollectionManagerCache::GetCachedColors() const
 			if (const TOptional<FLinearColor> CollectionColor = Collection->GetCollectionColor())
 			{
 				// Only add if not already present (ignores near matches too)
-				const bool bExists = CachedColors_Internal.FindByPredicate([CurrentColor = CollectionColor.GetValue()](const FLinearColor& Color) { return CurrentColor.Equals(Color); });
+				const bool bExists = CachedColors_Internal.ContainsByPredicate([CurrentColor = CollectionColor.GetValue()](const FLinearColor& Color) { return CurrentColor.Equals(Color); });
 				if (!bExists)
 				{
 					CachedColors_Internal.Add(CollectionColor.GetValue());
