Optimize collection files reading
  - Parallelize collection loading in FCollectionManager::LoadCollections
  - Parallelize FName registration to cut on hashing cost in Collection::Load
  - Leverage FStringView where possible in Collection::Load

  - 987ms -> 238ms for collection loading during editor startup
  - 987ms -> 696ms with -onethread

#rnx
#tests generated md5 of all loaded collections content before and after the change

diff --git a/Engine/Source/Developer/CollectionManager/Private/Collection.cpp b/Engine/Source/Developer/CollectionManager/Private/Collection.cpp
--- a/Engine/Source/Developer/CollectionManager/Private/Collection.cpp
+++ b/Engine/Source/Developer/CollectionManager/Private/Collection.cpp
@@ -14,6 +14,9 @@
 #include "ISourceControlModule.h"
 #include "Misc/TextFilterExpressionEvaluator.h"
 #include "Misc/EngineBuildSettings.h"
+#include "Misc/ScopeRWLock.h"
+#include "Async/ParallelFor.h"
+#include "String/ParseLines.h"
 
 #define LOCTEXT_NAMESPACE "CollectionManager"
 
@@ -67,6 +70,7 @@ TSharedRef<FCollection> FCollection::Clone(const FString& InFilename, bool InUse
 
 bool FCollection::Load(FText& OutError)
 {
+	TRACE_CPUPROFILER_EVENT_SCOPE(FCollection::Load)
 	Empty();
 
 	FString FullFileContentsString;
@@ -76,25 +80,23 @@ bool FCollection::Load(FText& OutError)
 		return false;
 	}
 
-	// Normalize line endings and parse into array
-	TArray<FString> FileContents;
-	FullFileContentsString.ReplaceInline(TEXT("\r"), TEXT(""));
-	FullFileContentsString.ParseIntoArray(FileContents, TEXT("\n"), /*bCullEmpty=*/false);
+	TArray<FStringView> FileContents;
+	UE::String::ParseLines(FullFileContentsString, [&FileContents](const FStringView& Line) { FileContents.Add(Line); });
 
-	if ( FileContents.Num() == 0 )
+	if (FileContents.Num() == 0)
 	{
 		// Empty file, assume static collection with no items
 		return true;
 	}
 
 	// Load the header from the contents array
-	TMap<FString,FString> HeaderPairs;
-	while ( FileContents.Num() )
+	TMap<FString, FString> HeaderPairs;
+	
+	int32 LineIndex = 0;
+	for (int32 Num = FileContents.Num(); LineIndex < Num; ++LineIndex)
 	{
-		// Pop the 0th element from the contents array and read it
-		FileContents[0].TrimStartAndEndInline();
-		const FString Line = FileContents[0];
-		FileContents.RemoveAt(0);
+		FStringView Line(FileContents[LineIndex]);
+		Line.TrimStartAndEndInline();
 
 		if (Line.Len() == 0)
 		{
@@ -102,11 +104,12 @@ bool FCollection::Load(FText& OutError)
 			break;
 		}
 
-		FString Key;
-		FString Value;
-		if ( Line.Split(TEXT(":"), &Key, &Value) )
+		FStringView::SizeType Offset;
+		if (Line.FindChar(TEXT(':'), Offset))
 		{
-			HeaderPairs.Add(Key, Value);
+			FString Key(Line.Left(Offset));
+			FString Value(Line.Right(Line.Len() - Offset - 1));
+			HeaderPairs.Emplace(MoveTemp(Key), MoveTemp(Value));
 		}
 	}
 
@@ -121,23 +124,34 @@ bool FCollection::Load(FText& OutError)
 	// Now load the content if the header load was successful
 	if (StorageMode == ECollectionStorageMode::Static)
 	{
-		// Static collection, a flat list of asset paths
-		for (FString& Line : FileContents)
-		{
-			Line.TrimStartAndEndInline();
+		const int32 NamesNum = FileContents.Num() - LineIndex;
 
-			if ( int32 Len = Line.Len() )
+		TArray<FName> FNames;
+		FNames.SetNum(NamesNum);
+
+		// Name hashing to register new FName takes time
+		// Process as much as possible in multiple threads
+		ParallelFor(
+			NamesNum,
+			[this, &FileContents, &LineIndex, &FNames](int32 LocalLineIndex)
 			{
-				AddObjectToCollection(FName(Len, *Line));
-			}
+				FStringView Line(FileContents[LineIndex + LocalLineIndex]);
+				FNames[LocalLineIndex] = FName(Line.TrimStartAndEnd());
+			},
+			// Do not pay for scheduling cost if number of items is too low
+			NamesNum < 1000 ? EParallelForFlags::ForceSingleThread : EParallelForFlags::None
+		);
+
+		// Static collection, a flat list of asset paths
+		for (FName& Name : FNames)
+		{
+			AddObjectToCollection(Name);
 		}
 	}
 	else
 	{
 		// Dynamic collection, a single query line
-		DynamicQueryText = (FileContents.Num() > 0) ? FileContents[0] : FString();
-
-		DynamicQueryText.TrimStartAndEndInline();
+		DynamicQueryText = (FileContents.Num() > LineIndex) ? FString(FileContents[LineIndex].TrimStartAndEnd()) : FString();
 	}
 
 	DiskSnapshot.TakeSnapshot(*this);
@@ -745,6 +759,14 @@ void FCollection::GetObjectDifferences(const TSet<FName>& BaseSet, const TSet<FN
 		}
 	}
 
+	// If both sets have the same number of items and nothing has been removed
+	// we can safely infer that both collections are equals without going
+	// over them a second time.
+	if (ObjectsRemoved.Num() == 0 && BaseSet.Num() == NewSet.Num())
+	{
+		return;
+	}
+
 	// Find the objects that were added compare to the base set
 	for (const FName& NewObjectName : NewSet)
 	{
diff --git a/Engine/Source/Developer/CollectionManager/Private/CollectionManager.cpp b/Engine/Source/Developer/CollectionManager/Private/CollectionManager.cpp
--- a/Engine/Source/Developer/CollectionManager/Private/CollectionManager.cpp
+++ b/Engine/Source/Developer/CollectionManager/Private/CollectionManager.cpp
@@ -7,6 +7,8 @@
 #include "CollectionManagerLog.h"
 #include "FileCache.h"
 #include "Misc/FileHelper.h"
+#include "Misc/ScopeRWLock.h"
+#include "Async/ParallelFor.h"
 
 #define LOCTEXT_NAMESPACE "CollectionManager"
 
@@ -1816,35 +1818,48 @@ bool FCollectionManager::TickFileCache(float InDeltaTime)
 
 void FCollectionManager::LoadCollections()
 {
+	TRACE_CPUPROFILER_EVENT_SCOPE(FCollectionManager::LoadCollections)
+
 	const double LoadStartTime = FPlatformTime::Seconds();
 	const int32 PrevNumCollections = AvailableCollections.Num();
 
-	for (int32 CacheIdx = 0; CacheIdx < ECollectionShareType::CST_All; ++CacheIdx)
-	{
-		const ECollectionShareType::Type ShareType = ECollectionShareType::Type(CacheIdx);
-		const FString& CollectionFolder = CollectionFolders[CacheIdx];
-		const FString WildCard = FString::Printf(TEXT("%s/*.%s"), *CollectionFolder, *CollectionExtension);
-
-		TArray<FString> Filenames;
-		IFileManager::Get().FindFiles(Filenames, *WildCard, true, false);
-
-		for (const FString& BaseFilename : Filenames)
+	FRWLock CollectionLock;
+	ParallelFor(
+		ECollectionShareType::CST_All,
+		[this, &CollectionLock](int32 CacheIdx)
 		{
-			const FString Filename = CollectionFolder / BaseFilename;
+			const ECollectionShareType::Type ShareType = ECollectionShareType::Type(CacheIdx);
 			const bool bUseSCC = ShouldUseSCC(ShareType);
+			const FString& CollectionFolder = CollectionFolders[CacheIdx];
+			const FString WildCard = FString::Printf(TEXT("%s/*.%s"), *CollectionFolder, *CollectionExtension);
 
-			FText LoadErrorText;
-			TSharedRef<FCollection> NewCollection = MakeShareable(new FCollection(Filename, bUseSCC, ECollectionStorageMode::Static));
-			if (NewCollection->Load(LoadErrorText))
-			{
-				AddCollection(NewCollection, ShareType);
-			}
-			else
-			{
-				UE_LOG(LogCollectionManager, Warning, TEXT("%s"), *LoadErrorText.ToString());
-			}
-		}
-	}
+			TArray<FString> Filenames;
+			IFileManager::Get().FindFiles(Filenames, *WildCard, true, false);
+
+			ParallelFor(
+				Filenames.Num(),
+				[this, &CollectionLock, &Filenames, &CollectionFolder, bUseSCC, ShareType](int32 FilenameIdx)
+				{
+					const FString& BaseFilename = Filenames[FilenameIdx];
+					const FString Filename = CollectionFolder / BaseFilename;
+
+					FText LoadErrorText;
+					TSharedRef<FCollection> NewCollection = MakeShareable(new FCollection(Filename, bUseSCC, ECollectionStorageMode::Static));
+					if (NewCollection->Load(LoadErrorText))
+					{
+						FRWScopeLock Lock(CollectionLock, SLT_Write);
+						AddCollection(NewCollection, ShareType);
+					}
+					else
+					{
+						UE_LOG(LogCollectionManager, Warning, TEXT("%s"), *LoadErrorText.ToString());
+					}
+				},
+				EParallelForFlags::Unbalanced
+			);
+		},
+		EParallelForFlags::Unbalanced
+	);
 
 	// AddCollection is assumed to be adding an empty collection, so also notify that collection cache that the collection has "changed" since loaded collections may not always be empty
 	CollectionCache.HandleCollectionChanged();
