#UE Add commandline option to avoid following redirectors in collections during startup. You sometimes want to make collections of redirectors when audting content

diff --git a/Engine/Source/Developer/CollectionManager/Private/CollectionManager.cpp b/Engine/Source/Developer/CollectionManager/Private/CollectionManager.cpp
--- a/Engine/Source/Developer/CollectionManager/Private/CollectionManager.cpp
+++ b/Engine/Source/Developer/CollectionManager/Private/CollectionManager.cpp
@@ -9,6 +9,7 @@
 #include "Misc/FileHelper.h"
 #include "Misc/ScopeRWLock.h"
 #include "Async/ParallelFor.h"
+#include "Misc/CommandLine.h"
 
 #define LOCTEXT_NAMESPACE "CollectionManager"
 
@@ -268,6 +269,8 @@ FCollectionManager::FCollectionManager()
 
 	CollectionExtension = TEXT("collection");
 
+	bNoFixupRedirectors = FParse::Param(FCommandLine::Get(), TEXT("NoFixupRedirectorsInCollections"));
+
 	LoadCollections();
 
 	// Watch for changes that may happen outside of the collection manager
@@ -1543,6 +1546,11 @@ bool FCollectionManager::IsValidParentCollection(FName CollectionName, ECollecti
 
 void FCollectionManager::HandleFixupRedirectors(ICollectionRedirectorFollower& InRedirectorFollower)
 {
+	if (bNoFixupRedirectors)
+	{
+		return;
+	}
+
 	const double LoadStartTime = FPlatformTime::Seconds();
 
 	TArray<TPair<FName, FName>> ObjectsToRename;
diff --git a/Engine/Source/Developer/CollectionManager/Private/CollectionManager.h b/Engine/Source/Developer/CollectionManager/Private/CollectionManager.h
--- a/Engine/Source/Developer/CollectionManager/Private/CollectionManager.h
+++ b/Engine/Source/Developer/CollectionManager/Private/CollectionManager.h
@@ -266,4 +266,7 @@ private:
 
 	/** When a collection checkin happens, use this event to add additional text to the changelist description */
 	FAddToCollectionCheckinDescriptionEvent AddToCollectionCheckinDescriptionEvent;
+
+	/** When true, redirectors will not be automatically followed in collections during startup */
+	bool bNoFixupRedirectors;
 };
