Editor: Added a project setting to allow specifying additional changelist tags for use when the editor auto-commits edits to collections

diff --git a/Engine/Source/Developer/CollectionManager/CollectionManager.Build.cs b/Engine/Source/Developer/CollectionManager/CollectionManager.Build.cs
--- a/Engine/Source/Developer/CollectionManager/CollectionManager.Build.cs
+++ b/Engine/Source/Developer/CollectionManager/CollectionManager.Build.cs
@@ -9,6 +9,7 @@ public class CollectionManager : ModuleRules
 		PrivateDependencyModuleNames.AddRange(
 			new string[] {
 				"Core",
+				"CoreUObject",
 				"DirectoryWatcher",
 				"SourceControl"
 			}
diff --git a/Engine/Source/Developer/CollectionManager/Private/CollectionManager.cpp b/Engine/Source/Developer/CollectionManager/Private/CollectionManager.cpp
--- a/Engine/Source/Developer/CollectionManager/Private/CollectionManager.cpp
+++ b/Engine/Source/Developer/CollectionManager/Private/CollectionManager.cpp
@@ -10,6 +10,7 @@
 #include "Misc/ScopeRWLock.h"
 #include "Async/ParallelFor.h"
 #include "Misc/CommandLine.h"
+#include "SourceControlPreferences.h"
 
 #define LOCTEXT_NAMESPACE "CollectionManager"
 
@@ -1984,7 +1985,27 @@ void FCollectionManager::ReplaceObjectInCollections(const FName& OldObjectPath,
 bool FCollectionManager::InternalSaveCollection(const TSharedRef<FCollection>& CollectionRef, FText& OutError)
 {
 	TArray<FText> AdditionalChangelistText;
+
+	// Give game specific editors a chance to add lines
 	AddToCollectionCheckinDescriptionEvent.Broadcast(CollectionRef->GetCollectionName(), AdditionalChangelistText);
+
+	// Give settings a chance to add lines
+	TArray<FString> SettingsLines;
+
+	const USourceControlPreferences* Settings = GetDefault<USourceControlPreferences>();
+	if (const FString* SpecificMatch = Settings->SpecificCollectionChangelistTags.Find(CollectionRef->GetCollectionName()))
+	{
+		// Parse input buffer into an array of lines
+		SpecificMatch->ParseIntoArrayLines(SettingsLines, /*bCullEmpty=*/ false);
+	}
+	SettingsLines.Append(Settings->CollectionChangelistTags);
+
+	for (const FString& OneSettingLine : SettingsLines)
+	{
+		AdditionalChangelistText.Add(FText::FromString(*OneSettingLine));
+	}
+
+	// Save the collection
 	return CollectionRef->Save(AdditionalChangelistText, OutError);
 }
 
diff --git a/Engine/Source/Developer/SourceControl/Public/SourceControlPreferences.h b/Engine/Source/Developer/SourceControl/Public/SourceControlPreferences.h
--- a/Engine/Source/Developer/SourceControl/Public/SourceControlPreferences.h
+++ b/Engine/Source/Developer/SourceControl/Public/SourceControlPreferences.h
@@ -29,4 +29,12 @@ public:
 	/** If enabled, deletes new files when reverted. */
 	UPROPERTY(config, EditAnywhere, Category = SourceControl, meta = (ToolTip = "Deletes new files when reverted."))
 	bool bShouldDeleteNewFilesOnRevert = true;
+
+	/** List of lines to add to any collection on checkin */
+	UPROPERTY(config, EditAnywhere, Category = SourceControl)
+	TArray<FString> CollectionChangelistTags;
+
+	/** Map of collection names and additional text to apply to changelist descriptions when checking them in */
+	UPROPERTY(config, EditAnywhere, Category = SourceControl, meta=(MultiLine=true))
+	TMap<FName, FString> SpecificCollectionChangelistTags;
 };
