Updating collection manager interface to use FSoftObjectPath for the paths of assets and objects rather than FName.

diff --git a/Engine/Plugins/Editor/AssetManagerEditor/Source/AssetManagerEditor/Private/ReferenceViewer/EdGraph_ReferenceViewer.cpp b/Engine/Plugins/Editor/AssetManagerEditor/Source/AssetManagerEditor/Private/ReferenceViewer/EdGraph_ReferenceViewer.cpp
--- a/Engine/Plugins/Editor/AssetManagerEditor/Source/AssetManagerEditor/Private/ReferenceViewer/EdGraph_ReferenceViewer.cpp
+++ b/Engine/Plugins/Editor/AssetManagerEditor/Source/AssetManagerEditor/Private/ReferenceViewer/EdGraph_ReferenceViewer.cpp
@@ -181,12 +181,12 @@ UEdGraphNode_Reference* UEdGraph_ReferenceViewer::ConstructNodes(const TArray<FA
 		if (ShouldFilterByCollection())
 		{
 			FCollectionManagerModule& CollectionManagerModule = FCollectionManagerModule::GetModule();
-			TArray<FName> AssetPaths;
+			TArray<FSoftObjectPath> AssetPaths;
 			CollectionManagerModule.Get().GetAssetsInCollection(CurrentCollectionFilter, ECollectionShareType::CST_All, AssetPaths);
 			CurrentCollectionPackages.Reserve(AssetPaths.Num());
-			for (FName AssetPath : AssetPaths)
+			for (const FSoftObjectPath& AssetPath  : AssetPaths)
 			{
-				CurrentCollectionPackages.Add(FName(*FPackageName::ObjectPathToPackageName(AssetPath.ToString())));
+				CurrentCollectionPackages.Add(AssetPath.GetLongPackageFName());
 			}
 		}
 
diff --git a/Engine/Plugins/Editor/ContentBrowser/ContentBrowserAliasDataSoure/Source/ContentBrowserAliasDataSource/Private/ContentBrowserAliasDataSource.cpp b/Engine/Plugins/Editor/ContentBrowser/ContentBrowserAliasDataSoure/Source/ContentBrowserAliasDataSource/Private/ContentBrowserAliasDataSource.cpp
--- a/Engine/Plugins/Editor/ContentBrowser/ContentBrowserAliasDataSoure/Source/ContentBrowserAliasDataSource/Private/ContentBrowserAliasDataSource.cpp
+++ b/Engine/Plugins/Editor/ContentBrowser/ContentBrowserAliasDataSoure/Source/ContentBrowserAliasDataSource/Private/ContentBrowserAliasDataSource.cpp
@@ -730,13 +730,13 @@ bool UContentBrowserAliasDataSource::UpdateThumbnail(const FContentBrowserItemDa
 	return ContentBrowserAssetData::UpdateItemThumbnail(this, InItem, InThumbnail);
 }
 
-bool UContentBrowserAliasDataSource::TryGetCollectionId(const FContentBrowserItemData& InItem, FName& OutCollectionId)
+bool UContentBrowserAliasDataSource::TryGetCollectionId(const FContentBrowserItemData& InItem, FSoftObjectPath& OutCollectionId)
 {
 	if (TSharedPtr<const FContentBrowserAliasItemDataPayload> AliasPayload = StaticCastSharedPtr<const FContentBrowserAliasItemDataPayload>(InItem.GetPayload()))
 	{
 		if (const FAliasData* AliasData = AllAliases.Find(AliasPayload->Alias))
 		{
-			OutCollectionId = AliasData->ObjectPath;
+			OutCollectionId = FSoftObjectPath(AliasData->ObjectPath);
 			return true;
 		}
 	}
@@ -794,7 +794,9 @@ FContentBrowserItemData UContentBrowserAliasDataSource::CreateAssetFileItem(cons
 	if (AliasData)
 	{
 		FName VirtualizedPath;
-		TryConvertInternalPathToVirtual(AliasData->ObjectPath, VirtualizedPath);
+PRAGMA_DISABLE_DEPRECATION_WARNINGS
+		TryConvertInternalPathToVirtual(AliasData->ObjectPath.ToFName(), VirtualizedPath);
+PRAGMA_ENABLE_DEPRECATION_WARNINGS
 
 		// Since AliasID is PackagePath/AssetName, AssetName should also be passed as the ItemName here. This provides the functionality of
 		// being able to have multiple aliases with the same display name, while still showing their original asset name in the tooltip.
diff --git a/Engine/Plugins/Editor/ContentBrowser/ContentBrowserAliasDataSoure/Source/ContentBrowserAliasDataSource/Public/ContentBrowserAliasDataSource.h b/Engine/Plugins/Editor/ContentBrowser/ContentBrowserAliasDataSoure/Source/ContentBrowserAliasDataSource/Public/ContentBrowserAliasDataSource.h
--- a/Engine/Plugins/Editor/ContentBrowser/ContentBrowserAliasDataSoure/Source/ContentBrowserAliasDataSource/Public/ContentBrowserAliasDataSource.h
+++ b/Engine/Plugins/Editor/ContentBrowser/ContentBrowserAliasDataSoure/Source/ContentBrowserAliasDataSource/Public/ContentBrowserAliasDataSource.h
@@ -82,7 +82,7 @@ public:
 	virtual bool AppendItemReference(const FContentBrowserItemData& InItem, FString& InOutStr) override;
 	virtual bool UpdateThumbnail(const FContentBrowserItemData& InItem, FAssetThumbnail& InThumbnail) override;
 
-	virtual bool TryGetCollectionId(const FContentBrowserItemData& InItem, FName& OutCollectionId) override;
+	virtual bool TryGetCollectionId(const FContentBrowserItemData& InItem, FSoftObjectPath& OutCollectionId) override;
 
 	// Legacy functions seem necessary for FrontendFilters to work
 	virtual bool Legacy_TryGetPackagePath(const FContentBrowserItemData& InItem, FName& OutPackagePath) override;
@@ -143,9 +143,7 @@ private:
 			PathBuilder << InAssetData.AssetName;
 			PackageName = FName(PathBuilder.ToView());
 
-			PathBuilder << TEXT('.');
-			PathBuilder << InAssetData.AssetName;
-			ObjectPath = FName(PathBuilder.ToView());
+			ObjectPath = FSoftObjectPath(PackageName, InAssetData.AssetName, {});
 		}
 
 		/** The source asset for this alias */
@@ -155,7 +153,7 @@ private:
 		/** PackagePath/SourceAssetName, /MyAliases/SomeAsset  */
 		FName PackageName;
 		/** PackageName.SourceAssetName, /MyAliases/SomeAsset.SomeAsset */
-		FName ObjectPath;
+		FSoftObjectPath ObjectPath;
 		/** A non-unique display name for this alias */
 		FName AliasName;
 		/** Whether this alias was generated from package metadata or manually through the C++ interface */
diff --git a/Engine/Plugins/Editor/ContentBrowser/ContentBrowserAssetDataSource/Source/ContentBrowserAssetDataSource/Private/ContentBrowserAssetDataSource.cpp b/Engine/Plugins/Editor/ContentBrowser/ContentBrowserAssetDataSource/Source/ContentBrowserAssetDataSource/Private/ContentBrowserAssetDataSource.cpp
--- a/Engine/Plugins/Editor/ContentBrowser/ContentBrowserAssetDataSource/Source/ContentBrowserAssetDataSource/Private/ContentBrowserAssetDataSource.cpp
+++ b/Engine/Plugins/Editor/ContentBrowser/ContentBrowserAssetDataSource/Source/ContentBrowserAssetDataSource/Private/ContentBrowserAssetDataSource.cpp
@@ -1581,13 +1581,11 @@ bool UContentBrowserAssetDataSource::HandleDragDropOnItem(const FContentBrowserI
 	return false;
 }
 
-bool UContentBrowserAssetDataSource::TryGetCollectionId(const FContentBrowserItemData& InItem, FName& OutCollectionId)
+bool UContentBrowserAssetDataSource::TryGetCollectionId(const FContentBrowserItemData& InItem, FSoftObjectPath& OutCollectionId)
 {
 	if (TSharedPtr<const FContentBrowserAssetFileItemDataPayload> AssetPayload = GetAssetFileItemPayload(InItem))
 	{
-PRAGMA_DISABLE_DEPRECATION_WARNINGS
-		OutCollectionId = AssetPayload->GetAssetData().ObjectPath;
-PRAGMA_ENABLE_DEPRECATION_WARNINGS
+		OutCollectionId = AssetPayload->GetAssetData().GetSoftObjectPath();
 		return true;
 	}
 	return false;
@@ -1661,14 +1659,9 @@ bool UContentBrowserAssetDataSource::GetObjectPathsForCollections(ICollectionMan
 	{
 		const ECollectionRecursionFlags::Flags CollectionRecursionMode = bIncludeChildCollections ? ECollectionRecursionFlags::SelfAndChildren : ECollectionRecursionFlags::Self;
 		
-		TArray<FName> TempNames;
 		for (const FCollectionNameType& CollectionNameType : InCollections)
 		{
-			TempNames.Reset();
-			CollectionManager->GetObjectsInCollection(CollectionNameType.Name, CollectionNameType.Type, TempNames, CollectionRecursionMode);
-PRAGMA_DISABLE_DEPRECATION_WARNINGS
-			OutObjectPaths.Append(UE::SoftObjectPath::Private::ConvertObjectPathNames(TempNames));
-PRAGMA_ENABLE_DEPRECATION_WARNINGS
+			CollectionManager->GetObjectsInCollection(CollectionNameType.Name, CollectionNameType.Type, OutObjectPaths, CollectionRecursionMode);
 		}
 
 		return true;
diff --git a/Engine/Plugins/Editor/ContentBrowser/ContentBrowserAssetDataSource/Source/ContentBrowserAssetDataSource/Public/ContentBrowserAssetDataSource.h b/Engine/Plugins/Editor/ContentBrowser/ContentBrowserAssetDataSource/Source/ContentBrowserAssetDataSource/Public/ContentBrowserAssetDataSource.h
--- a/Engine/Plugins/Editor/ContentBrowser/ContentBrowserAssetDataSource/Source/ContentBrowserAssetDataSource/Public/ContentBrowserAssetDataSource.h
+++ b/Engine/Plugins/Editor/ContentBrowser/ContentBrowserAssetDataSource/Source/ContentBrowserAssetDataSource/Public/ContentBrowserAssetDataSource.h
@@ -246,7 +246,7 @@ public:
 
 	virtual bool HandleDragDropOnItem(const FContentBrowserItemData& InItem, const FDragDropEvent& InDragDropEvent) override;
 
-	virtual bool TryGetCollectionId(const FContentBrowserItemData& InItem, FName& OutCollectionId) override;
+	virtual bool TryGetCollectionId(const FContentBrowserItemData& InItem, FSoftObjectPath& OutCollectionId) override;
 
 	virtual bool Legacy_TryGetPackagePath(const FContentBrowserItemData& InItem, FName& OutPackagePath) override;
 
diff --git a/Engine/Plugins/Editor/ContentBrowser/ContentBrowserClassDataSource/Source/ContentBrowserClassDataSource/Private/ContentBrowserClassDataSource.cpp b/Engine/Plugins/Editor/ContentBrowser/ContentBrowserClassDataSource/Source/ContentBrowserClassDataSource/Private/ContentBrowserClassDataSource.cpp
--- a/Engine/Plugins/Editor/ContentBrowser/ContentBrowserClassDataSource/Source/ContentBrowserClassDataSource/Private/ContentBrowserClassDataSource.cpp
+++ b/Engine/Plugins/Editor/ContentBrowser/ContentBrowserClassDataSource/Source/ContentBrowserClassDataSource/Private/ContentBrowserClassDataSource.cpp
@@ -285,10 +285,10 @@ void UContentBrowserClassDataSource::CompileFilter(const FName InPath, const FCo
 
 		if (ChildClassObjects.Num() > 0)
 		{
-			TSet<FName> ClassPathsToInclude;
+			TSet<FTopLevelAssetPath> ClassPathsToInclude;
 			if (CollectionFilter)
 			{
-				TArray<FName> ClassPathsForCollections;
+				TArray<FTopLevelAssetPath> ClassPathsForCollections;
 				if (GetClassPathsForCollections(CollectionFilter->SelectedCollections, CollectionFilter->bIncludeChildCollections, ClassPathsForCollections) && ClassPathsForCollections.Num() == 0)
 				{
 					// If we had collections but they contained no classes then we can bail as nothing will pass the filter
@@ -300,7 +300,7 @@ void UContentBrowserClassDataSource::CompileFilter(const FName InPath, const FCo
 
 			for (UClass* ChildClassObject : ChildClassObjects)
 			{
-				const bool bPassesInclusiveFilter = ClassPathsToInclude.Num() == 0 || ClassPathsToInclude.Contains(*ChildClassObject->GetPathName());
+				const bool bPassesInclusiveFilter = ClassPathsToInclude.Num() == 0 || ClassPathsToInclude.Contains(FTopLevelAssetPath(ChildClassObject));
 				const bool bPassesPermissionCheck = !ClassPermissionList || ClassPermissionList->PassesFilter(ChildClassObject->GetClassPathName().ToString());
 
 				if (bPassesInclusiveFilter && bPassesPermissionCheck)
@@ -512,13 +512,11 @@ bool UContentBrowserClassDataSource::UpdateThumbnail(const FContentBrowserItemDa
 	return ContentBrowserClassData::UpdateItemThumbnail(this, InItem, InThumbnail);
 }
 
-bool UContentBrowserClassDataSource::TryGetCollectionId(const FContentBrowserItemData& InItem, FName& OutCollectionId)
+bool UContentBrowserClassDataSource::TryGetCollectionId(const FContentBrowserItemData& InItem, FSoftObjectPath& OutCollectionId)
 {
 	if (TSharedPtr<const FContentBrowserClassFileItemDataPayload> ClassPayload = GetClassFileItemPayload(InItem))
 	{
-PRAGMA_DISABLE_DEPRECATION_WARNINGS
-		OutCollectionId = ClassPayload->GetAssetData().ObjectPath;
-PRAGMA_ENABLE_DEPRECATION_WARNINGS
+		OutCollectionId = FSoftObjectPath(ClassPayload->GetAssetData().GetSoftObjectPath());
 		return true;
 	}
 	return false;
@@ -564,7 +562,7 @@ bool UContentBrowserClassDataSource::IsKnownClassPath(const FName InPackagePath)
 	return FStringView(PackagePathStr).StartsWith(TEXT("/Classes_"));
 }
 
-bool UContentBrowserClassDataSource::GetClassPathsForCollections(TArrayView<const FCollectionNameType> InCollections, const bool bIncludeChildCollections, TArray<FName>& OutClassPaths)
+bool UContentBrowserClassDataSource::GetClassPathsForCollections(TArrayView<const FCollectionNameType> InCollections, const bool bIncludeChildCollections, TArray<FTopLevelAssetPath>& OutClassPaths)
 {
 	if (InCollections.Num() > 0)
 	{
diff --git a/Engine/Plugins/Editor/ContentBrowser/ContentBrowserClassDataSource/Source/ContentBrowserClassDataSource/Public/ContentBrowserClassDataSource.h b/Engine/Plugins/Editor/ContentBrowser/ContentBrowserClassDataSource/Source/ContentBrowserClassDataSource/Public/ContentBrowserClassDataSource.h
--- a/Engine/Plugins/Editor/ContentBrowser/ContentBrowserClassDataSource/Source/ContentBrowserClassDataSource/Public/ContentBrowserClassDataSource.h
+++ b/Engine/Plugins/Editor/ContentBrowser/ContentBrowserClassDataSource/Source/ContentBrowserClassDataSource/Public/ContentBrowserClassDataSource.h
@@ -66,7 +66,7 @@ public:
 
 	virtual bool UpdateThumbnail(const FContentBrowserItemData& InItem, FAssetThumbnail& InThumbnail) override;
 
-	virtual bool TryGetCollectionId(const FContentBrowserItemData& InItem, FName& OutCollectionId) override;
+	virtual bool TryGetCollectionId(const FContentBrowserItemData& InItem, FSoftObjectPath& OutCollectionId) override;
 
 	virtual bool Legacy_TryGetPackagePath(const FContentBrowserItemData& InItem, FName& OutPackagePath) override;
 
@@ -83,7 +83,7 @@ protected:
 private:
 	bool IsKnownClassPath(const FName InPackagePath) const;
 
-	bool GetClassPathsForCollections(TArrayView<const FCollectionNameType> InCollections, const bool bIncludeChildCollections, TArray<FName>& OutClassPaths);
+	bool GetClassPathsForCollections(TArrayView<const FCollectionNameType> InCollections, const bool bIncludeChildCollections, TArray<FTopLevelAssetPath>& OutClassPaths);
 
 	FContentBrowserItemData CreateClassFolderItem(const FName InFolderPath);
 
diff --git a/Engine/Plugins/FX/Niagara/Source/NiagaraEditor/Private/Commandlets/NiagaraDumpBytecodeCommandlet.cpp b/Engine/Plugins/FX/Niagara/Source/NiagaraEditor/Private/Commandlets/NiagaraDumpBytecodeCommandlet.cpp
--- a/Engine/Plugins/FX/Niagara/Source/NiagaraEditor/Private/Commandlets/NiagaraDumpBytecodeCommandlet.cpp
+++ b/Engine/Plugins/FX/Niagara/Source/NiagaraEditor/Private/Commandlets/NiagaraDumpBytecodeCommandlet.cpp
@@ -92,9 +92,7 @@ void UNiagaraDumpByteCodeCommandlet::ProcessNiagaraScripts()
 	if (!FilterCollection.IsEmpty())
 	{
 		FCollectionManagerModule& CollectionManagerModule = FCollectionManagerModule::GetModule();
-PRAGMA_DISABLE_DEPRECATION_WARNINGS
-		CollectionManagerModule.Get().GetObjectsInCollection(FName(*FilterCollection), ECollectionShareType::CST_All, Filter.ObjectPaths, ECollectionRecursionFlags::SelfAndChildren);
-PRAGMA_ENABLE_DEPRECATION_WARNINGS
+		CollectionManagerModule.Get().GetObjectsInCollection(FName(*FilterCollection), ECollectionShareType::CST_All, Filter.SoftObjectPaths, ECollectionRecursionFlags::SelfAndChildren);
 	}
 
 	TArray<FAssetData> AssetList;
diff --git a/Engine/Plugins/FX/Niagara/Source/NiagaraEditor/Private/Commandlets/NiagaraDumpModuleInfoCommandlet.cpp b/Engine/Plugins/FX/Niagara/Source/NiagaraEditor/Private/Commandlets/NiagaraDumpModuleInfoCommandlet.cpp
--- a/Engine/Plugins/FX/Niagara/Source/NiagaraEditor/Private/Commandlets/NiagaraDumpModuleInfoCommandlet.cpp
+++ b/Engine/Plugins/FX/Niagara/Source/NiagaraEditor/Private/Commandlets/NiagaraDumpModuleInfoCommandlet.cpp
@@ -85,9 +85,7 @@ void UNiagaraDumpModuleInfoCommandlet::ProcessNiagaraParameterDefinitions()
 	if (!FilterCollection.IsEmpty())
 	{
 		FCollectionManagerModule& CollectionManagerModule = FCollectionManagerModule::GetModule();
-PRAGMA_DISABLE_DEPRECATION_WARNINGS
-		CollectionManagerModule.Get().GetObjectsInCollection(FName(*FilterCollection), ECollectionShareType::CST_All, Filter.ObjectPaths, ECollectionRecursionFlags::SelfAndChildren);
-PRAGMA_ENABLE_DEPRECATION_WARNINGS
+		CollectionManagerModule.Get().GetObjectsInCollection(FName(*FilterCollection), ECollectionShareType::CST_All, Filter.SoftObjectPaths, ECollectionRecursionFlags::SelfAndChildren);
 	}
 
 	TArray<FAssetData> AssetList;
@@ -189,9 +187,7 @@ void UNiagaraDumpModuleInfoCommandlet::ProcessNiagaraScripts()
 	if (!FilterCollection.IsEmpty())
 	{
 		FCollectionManagerModule& CollectionManagerModule = FCollectionManagerModule::GetModule();
-PRAGMA_DISABLE_DEPRECATION_WARNINGS
-		CollectionManagerModule.Get().GetObjectsInCollection(FName(*FilterCollection), ECollectionShareType::CST_All, Filter.ObjectPaths, ECollectionRecursionFlags::SelfAndChildren);
-PRAGMA_ENABLE_DEPRECATION_WARNINGS
+		CollectionManagerModule.Get().GetObjectsInCollection(FName(*FilterCollection), ECollectionShareType::CST_All, Filter.SoftObjectPaths, ECollectionRecursionFlags::SelfAndChildren);
 	}
 
 	TArray<FAssetData> AssetList;
diff --git a/Engine/Plugins/FX/Niagara/Source/NiagaraEditor/Private/Commandlets/NiagaraSystemAuditCommandlet.cpp b/Engine/Plugins/FX/Niagara/Source/NiagaraEditor/Private/Commandlets/NiagaraSystemAuditCommandlet.cpp
--- a/Engine/Plugins/FX/Niagara/Source/NiagaraEditor/Private/Commandlets/NiagaraSystemAuditCommandlet.cpp
+++ b/Engine/Plugins/FX/Niagara/Source/NiagaraEditor/Private/Commandlets/NiagaraSystemAuditCommandlet.cpp
@@ -126,9 +126,7 @@ bool UNiagaraSystemAuditCommandlet::ProcessNiagaraSystems()
 	if (!FilterCollection.IsEmpty())
 	{
 		FCollectionManagerModule& CollectionManagerModule = FCollectionManagerModule::GetModule();
-PRAGMA_DISABLE_DEPRECATION_WARNINGS
-		CollectionManagerModule.Get().GetObjectsInCollection(FName(*FilterCollection), ECollectionShareType::CST_All, Filter.ObjectPaths, ECollectionRecursionFlags::SelfAndChildren);
-PRAGMA_ENABLE_DEPRECATION_WARNINGS
+		CollectionManagerModule.Get().GetObjectsInCollection(FName(*FilterCollection), ECollectionShareType::CST_All, Filter.SoftObjectPaths, ECollectionRecursionFlags::SelfAndChildren);
 	}
 
 	TArray<FAssetData> AssetList;
diff --git a/Engine/Plugins/Runtime/AssetTags/Source/AssetTags/Private/AssetTagsSubsystem.cpp b/Engine/Plugins/Runtime/AssetTags/Source/AssetTags/Private/AssetTagsSubsystem.cpp
--- a/Engine/Plugins/Runtime/AssetTags/Source/AssetTags/Private/AssetTagsSubsystem.cpp
+++ b/Engine/Plugins/Runtime/AssetTags/Source/AssetTags/Private/AssetTagsSubsystem.cpp
@@ -170,12 +170,5 @@
 
 bool UAssetTagsSubsystem::K2_AddAssetToCollection(const FName Name, const FSoftObjectPath& AssetPath)
 {
-PRAGMA_DISABLE_DEPRECATION_WARNINGS
-	return AddAssetToCollection(Name, AssetPath.ToFName());
-PRAGMA_ENABLE_DEPRECATION_WARNINGS
-}
-
-bool UAssetTagsSubsystem::AddAssetToCollection(const FName Name, const FName AssetPathName)
-{
 	ICollectionManager& CollectionManager = FCollectionManagerModule::GetModule().Get();
 
@@ -185,7 +178,7 @@
 		return false;
 	}
 
-	if (!CollectionManager.AddToCollection(ResolvedNameAndType.Name, ResolvedNameAndType.Type, AssetPathName))
+	if (!CollectionManager.AddToCollection(ResolvedNameAndType.Name, ResolvedNameAndType.Type, AssetPath))
 	{
 		AssetTagsSubsystemUtil::LogLastCollectionManagerError(CollectionManager);
 		return false;
@@ -194,11 +187,16 @@
 	return true;
 }
 
+bool UAssetTagsSubsystem::AddAssetToCollection(const FName Name, const FName AssetPathName)
+{
+PRAGMA_DISABLE_DEPRECATION_WARNINGS
+	return K2_AddAssetToCollection(Name, FSoftObjectPath(AssetPathName));
+PRAGMA_ENABLE_DEPRECATION_WARNINGS
+}
+
 bool UAssetTagsSubsystem::AddAssetDataToCollection(const FName Name, const FAssetData& AssetData)
 {
-PRAGMA_DISABLE_DEPRECATION_WARNINGS
-	return AddAssetToCollection(Name, AssetData.ObjectPath);
-PRAGMA_ENABLE_DEPRECATION_WARNINGS
+	return K2_AddAssetToCollection(Name, AssetData.GetSoftObjectPath());
 }
 
 bool UAssetTagsSubsystem::AddAssetPtrToCollection(const FName Name, const UObject* AssetPtr)
@@ -210,12 +208,5 @@
 
 bool UAssetTagsSubsystem::K2_AddAssetsToCollection(const FName Name, const TArray<FSoftObjectPath>& AssetPaths)
 {
-PRAGMA_DISABLE_DEPRECATION_WARNINGS
-	return AddAssetsToCollection(Name, UE::SoftObjectPath::Private::ConvertSoftObjectPaths(AssetPaths));
-PRAGMA_ENABLE_DEPRECATION_WARNINGS
-}
-
-bool UAssetTagsSubsystem::AddAssetsToCollection(const FName Name, const TArray<FName>& AssetPathNames)
-{
 	ICollectionManager& CollectionManager = FCollectionManagerModule::GetModule().Get();
 
@@ -225,7 +216,7 @@
 		return false;
 	}
 
-	if (!CollectionManager.AddToCollection(ResolvedNameAndType.Name, ResolvedNameAndType.Type, AssetPathNames))
+	if (!CollectionManager.AddToCollection(ResolvedNameAndType.Name, ResolvedNameAndType.Type, AssetPaths))
 	{
 		AssetTagsSubsystemUtil::LogLastCollectionManagerError(CollectionManager);
 		return false;
@@ -234,40 +225,36 @@
 	return true;
 }
 
+bool UAssetTagsSubsystem::AddAssetsToCollection(const FName Name, const TArray<FName>& AssetPathNames)
+{
+PRAGMA_DISABLE_DEPRECATION_WARNINGS
+	return K2_AddAssetsToCollection(Name, UE::SoftObjectPath::Private::ConvertObjectPathNames(AssetPathNames));
+PRAGMA_ENABLE_DEPRECATION_WARNINGS
+}
+
 bool UAssetTagsSubsystem::AddAssetDatasToCollection(const FName Name, const TArray<FAssetData>& AssetDatas)
 {
-	TArray<FName> AssetPathNames;
-	AssetPathNames.Reserve(AssetDatas.Num());
+	TArray<FSoftObjectPath> AssetPaths;
+	AssetPaths.Reserve(AssetDatas.Num());
 	for (const FAssetData& AssetData : AssetDatas)
 	{
-PRAGMA_DISABLE_DEPRECATION_WARNINGS
-		AssetPathNames.Add(AssetData.ObjectPath);
+		AssetPaths.Add(AssetData.GetSoftObjectPath());
 	}
-	return AddAssetsToCollection(Name, AssetPathNames);
-PRAGMA_ENABLE_DEPRECATION_WARNINGS
+	return K2_AddAssetsToCollection(Name, AssetPaths);
 }
 
 bool UAssetTagsSubsystem::AddAssetPtrsToCollection(const FName Name, const TArray<UObject*>& AssetPtrs)
 {
-	TArray<FName> AssetPathNames;
-	AssetPathNames.Reserve(AssetPtrs.Num());
+	TArray<FSoftObjectPath> AssetPaths;
+	AssetPaths.Reserve(AssetPtrs.Num());
 	for (const UObject* AssetPtr : AssetPtrs)
 	{
-		AssetPathNames.Add(AssetPtr ? FName(*AssetPtr->GetPathName()) : NAME_None);
+		AssetPaths.Add(FSoftObjectPath(AssetPtr));
 	}
-PRAGMA_DISABLE_DEPRECATION_WARNINGS
-	return AddAssetsToCollection(Name, AssetPathNames);
-PRAGMA_ENABLE_DEPRECATION_WARNINGS
+	return K2_AddAssetsToCollection(Name, AssetPaths);
 }
 
 bool UAssetTagsSubsystem::K2_RemoveAssetFromCollection(const FName Name, const FSoftObjectPath& AssetPath)
 {
-PRAGMA_DISABLE_DEPRECATION_WARNINGS
-	return RemoveAssetFromCollection(Name, AssetPath.ToFName());
-PRAGMA_ENABLE_DEPRECATION_WARNINGS
-}
-
-bool UAssetTagsSubsystem::RemoveAssetFromCollection(const FName Name, const FName AssetPathName)
-{
 	ICollectionManager& CollectionManager = FCollectionManagerModule::GetModule().Get();
 
@@ -277,7 +264,7 @@
 		return false;
 	}
 
-	if (!CollectionManager.RemoveFromCollection(ResolvedNameAndType.Name, ResolvedNameAndType.Type, AssetPathName))
+	if (!CollectionManager.RemoveFromCollection(ResolvedNameAndType.Name, ResolvedNameAndType.Type, AssetPath))
 	{
 		AssetTagsSubsystemUtil::LogLastCollectionManagerError(CollectionManager);
 		return false;
@@ -286,11 +273,16 @@
 	return true;
 }
 
+bool UAssetTagsSubsystem::RemoveAssetFromCollection(const FName Name, const FName AssetPathName)
+{
+PRAGMA_DISABLE_DEPRECATION_WARNINGS
+	return K2_RemoveAssetFromCollection(Name, FSoftObjectPath(AssetPathName));
+PRAGMA_ENABLE_DEPRECATION_WARNINGS
+}
+
 bool UAssetTagsSubsystem::RemoveAssetDataFromCollection(const FName Name, const FAssetData& AssetData)
 {
-PRAGMA_DISABLE_DEPRECATION_WARNINGS
-	return RemoveAssetFromCollection(Name, AssetData.ObjectPath);
-PRAGMA_ENABLE_DEPRECATION_WARNINGS
+	return K2_RemoveAssetFromCollection(Name, AssetData.GetSoftObjectPath());
 }
 
 bool UAssetTagsSubsystem::RemoveAssetPtrFromCollection(const FName Name, const UObject* AssetPtr)
@@ -302,12 +294,5 @@
 
 bool UAssetTagsSubsystem::K2_RemoveAssetsFromCollection(const FName Name, const TArray<FSoftObjectPath>& AssetPaths)
 {
-PRAGMA_DISABLE_DEPRECATION_WARNINGS
-	return RemoveAssetsFromCollection(Name, UE::SoftObjectPath::Private::ConvertSoftObjectPaths(AssetPaths));
-PRAGMA_ENABLE_DEPRECATION_WARNINGS
-}
-
-bool UAssetTagsSubsystem::RemoveAssetsFromCollection(const FName Name, const TArray<FName>& AssetPathNames)
-{
 	ICollectionManager& CollectionManager = FCollectionManagerModule::GetModule().Get();
 
@@ -317,7 +302,7 @@
 		return false;
 	}
 
-	if (!CollectionManager.RemoveFromCollection(ResolvedNameAndType.Name, ResolvedNameAndType.Type, AssetPathNames))
+	if (!CollectionManager.RemoveFromCollection(ResolvedNameAndType.Name, ResolvedNameAndType.Type, AssetPaths))
 	{
 		AssetTagsSubsystemUtil::LogLastCollectionManagerError(CollectionManager);
 		return false;
@@ -326,30 +311,33 @@
 	return true;
 }
 
+bool UAssetTagsSubsystem::RemoveAssetsFromCollection(const FName Name, const TArray<FName>& AssetPathNames)
+{
+PRAGMA_DISABLE_DEPRECATION_WARNINGS
+	return K2_RemoveAssetsFromCollection(Name, UE::SoftObjectPath::Private::ConvertObjectPathNames(AssetPathNames));
+PRAGMA_ENABLE_DEPRECATION_WARNINGS
+}
+
 bool UAssetTagsSubsystem::RemoveAssetDatasFromCollection(const FName Name, const TArray<FAssetData>& AssetDatas)
 {
-	TArray<FName> AssetPathNames;
-	AssetPathNames.Reserve(AssetDatas.Num());
+	TArray<FSoftObjectPath> AssetPaths;
+	AssetPaths.Reserve(AssetDatas.Num());
 	for (const FAssetData& AssetData : AssetDatas)
 	{
-PRAGMA_DISABLE_DEPRECATION_WARNINGS
-		AssetPathNames.Add(AssetData.ObjectPath);
+		AssetPaths.Add(AssetData.GetSoftObjectPath());
 	}
-	return RemoveAssetsFromCollection(Name, AssetPathNames);
-PRAGMA_ENABLE_DEPRECATION_WARNINGS
+	return K2_RemoveAssetsFromCollection(Name, AssetPaths);
 }
 
 bool UAssetTagsSubsystem::RemoveAssetPtrsFromCollection(const FName Name, const TArray<UObject*>& AssetPtrs)
 {
-	TArray<FName> AssetPathNames;
-	AssetPathNames.Reserve(AssetPtrs.Num());
+	TArray<FSoftObjectPath> AssetPaths;
+	AssetPaths.Reserve(AssetPtrs.Num());
 	for (const UObject* AssetPtr : AssetPtrs)
 	{
-		AssetPathNames.Add(AssetPtr ? FName(*AssetPtr->GetPathName()) : NAME_None);
+		AssetPaths.Add(FSoftObjectPath(AssetPtr));
 	}
-PRAGMA_DISABLE_DEPRECATION_WARNINGS
-	return RemoveAssetsFromCollection(Name, AssetPathNames);
-PRAGMA_ENABLE_DEPRECATION_WARNINGS
+	return K2_RemoveAssetsFromCollection(Name, AssetPaths);
 }
 
 #endif	// WITH_EDITOR
@@ -428,12 +416,8 @@
 		const FCollectionNameType ResolvedNameAndType = AssetTagsSubsystemUtil::FindCollectionByName(CollectionManager, Name);
 		if (!ResolvedNameAndType.Name.IsNone())
 		{
-			TArray<FName> AssetPathNames;
-			CollectionManager.GetAssetsInCollection(ResolvedNameAndType.Name, ResolvedNameAndType.Type, AssetPathNames);
-
-PRAGMA_DISABLE_DEPRECATION_WARNINGS
-			TArray<FSoftObjectPath> AssetPaths = UE::SoftObjectPath::Private::ConvertObjectPathNames(AssetPathNames);
-PRAGMA_ENABLE_DEPRECATION_WARNINGS
+			TArray<FSoftObjectPath> AssetPaths;
+			CollectionManager.GetAssetsInCollection(ResolvedNameAndType.Name, ResolvedNameAndType.Type, AssetPaths);
 			Assets.Reserve(AssetPaths.Num());
 			for (const FSoftObjectPath& AssetPath : AssetPaths)
 			{
@@ -463,12 +447,5 @@
 
 TArray<FName> UAssetTagsSubsystem::K2_GetCollectionsContainingAsset(const FSoftObjectPath& AssetPath)
 {
-PRAGMA_DISABLE_DEPRECATION_WARNINGS
-	return GetCollectionsContainingAsset(AssetPath.ToFName());
-PRAGMA_ENABLE_DEPRECATION_WARNINGS
-}
-
-TArray<FName> UAssetTagsSubsystem::GetCollectionsContainingAsset(const FName AssetPathName)
-{
 	TArray<FName> CollectionNames;
 
@@ -477,7 +454,7 @@
 		ICollectionManager& CollectionManager = FCollectionManagerModule::GetModule().Get();
 
 		TArray<FCollectionNameType> CollectionNamesAndTypes;
-		CollectionManager.GetCollectionsContainingObject(AssetPathName, CollectionNamesAndTypes);
+		CollectionManager.GetCollectionsContainingObject(AssetPath, CollectionNamesAndTypes);
 
 		CollectionNames.Reserve(CollectionNamesAndTypes.Num());
 		for (const FCollectionNameType& CollectionNameAndType : CollectionNamesAndTypes)
@@ -490,7 +467,7 @@
 		IAssetRegistry& AssetRegistry = FModuleManager::LoadModuleChecked<FAssetRegistryModule>(AssetRegistryConstants::ModuleName).Get();
 
 		const bool bIncludeOnlyOnDiskAssets = true; // Collection tags are added at cook-time, so we *must* search the on-disk version of the tags (from the asset registry)
-		const FAssetData AssetData = AssetRegistry.GetAssetByObjectPath(AssetPathName, bIncludeOnlyOnDiskAssets);
+		const FAssetData AssetData = AssetRegistry.GetAssetByObjectPath(AssetPath, bIncludeOnlyOnDiskAssets);
 		if (AssetData.IsValid())
 		{
 			const TCHAR* CollectionTagPrefix = FAssetData::GetCollectionTagPrefix();
@@ -513,22 +490,25 @@
 	return CollectionNames;
 }
 
+TArray<FName> UAssetTagsSubsystem::GetCollectionsContainingAsset(const FName AssetPathName)
+{
+PRAGMA_DISABLE_DEPRECATION_WARNINGS
+	return K2_GetCollectionsContainingAsset(FSoftObjectPath(AssetPathName));
+PRAGMA_ENABLE_DEPRECATION_WARNINGS
+}
+
 TArray<FName> UAssetTagsSubsystem::GetCollectionsContainingAssetData(const FAssetData& AssetData)
 {
-	// Note: Use the path name version as the common implementation as:
-	//  1) The path name is always required for the collection manager implementation
+	// Note: Use the path version as the common implementation as:
+	//  1) The path is always required for the collection manager implementation
 	//  2) The FAssetData for the asset registry implementation *must* come from the asset registry (as the tags are added at cook-time, and missing if FAssetData is generated from a UObject* at runtime)
-PRAGMA_DISABLE_DEPRECATION_WARNINGS
-	return GetCollectionsContainingAsset(AssetData.ObjectPath);
-PRAGMA_ENABLE_DEPRECATION_WARNINGS
+	return K2_GetCollectionsContainingAsset(AssetData.GetSoftObjectPath());
 }
 
 TArray<FName> UAssetTagsSubsystem::GetCollectionsContainingAssetPtr(const UObject* AssetPtr)
 {
-	// Note: Use the path name version as the common implementation as:
-	//  1) The path name is always required for the collection manager implementation
+	// Note: Use the path version as the common implementation as:
+	//  1) The path is always required for the collection manager implementation
 	//  2) The FAssetData for the asset registry implementation *must* come from the asset registry (as the tags are added at cook-time, and missing if FAssetData is generated from a UObject* at runtime)
-PRAGMA_DISABLE_DEPRECATION_WARNINGS
-	return GetCollectionsContainingAsset(AssetPtr ? FName(*AssetPtr->GetPathName()) : NAME_None);
-PRAGMA_ENABLE_DEPRECATION_WARNINGS
+	return K2_GetCollectionsContainingAsset(FSoftObjectPath(AssetPtr));
 }
diff --git a/Engine/Plugins/VirtualProduction/RemoteControl/Source/WebRemoteControl/Private/RemoteControlModels.h b/Engine/Plugins/VirtualProduction/RemoteControl/Source/WebRemoteControl/Private/RemoteControlModels.h
--- a/Engine/Plugins/VirtualProduction/RemoteControl/Source/WebRemoteControl/Private/RemoteControlModels.h
+++ b/Engine/Plugins/VirtualProduction/RemoteControl/Source/WebRemoteControl/Private/RemoteControlModels.h
@@ -551,7 +551,9 @@ struct FRCShortPresetDescription
 		}
 
 		ID = PresetId.ToString();
+PRAGMA_DISABLE_DEPRECATION_WARNINGS
 		Path = PresetAsset.ObjectPath;
+PRAGMA_ENABLE_DEPRECATION_WARNINGS
 	}
 
 	FRCShortPresetDescription(const URemoteControlPreset* InPreset)
@@ -602,7 +604,9 @@ struct FRCAssetDescription
 	FRCAssetDescription(const FAssetData& InAsset)
 		: Name(InAsset.AssetName)
 		, Class(InAsset.AssetClassPath)
+PRAGMA_DISABLE_DEPRECATION_WARNINGS
 		, Path(InAsset.ObjectPath)
+PRAGMA_ENABLE_DEPRECATION_WARNINGS
 	{
 		Metadata = RemoteControlModels::SanitizeAssetMetadata(InAsset.TagsAndValues.CopyMap());
 	}
diff --git a/Engine/Plugins/VirtualProduction/RemoteControl/Source/WebRemoteControl/Private/WebRemoteControlEditorRoutes.cpp b/Engine/Plugins/VirtualProduction/RemoteControl/Source/WebRemoteControl/Private/WebRemoteControlEditorRoutes.cpp
--- a/Engine/Plugins/VirtualProduction/RemoteControl/Source/WebRemoteControl/Private/WebRemoteControlEditorRoutes.cpp
+++ b/Engine/Plugins/VirtualProduction/RemoteControl/Source/WebRemoteControl/Private/WebRemoteControlEditorRoutes.cpp
@@ -95,7 +95,7 @@ bool FWebRemoteControlEditorRoutes::HandleGetThumbnailRoute(const FHttpServerReq
 
 	FAssetRegistryModule& AssetRegistryModule = FModuleManager::Get().LoadModuleChecked<FAssetRegistryModule>(TEXT("AssetRegistry"));
 	FARFilter Filter;
-	Filter.ObjectPaths.Add(*GetThumbnailRequest.ObjectPath);
+	Filter.SoftObjectPaths.Add(GetThumbnailRequest.ObjectPath);
 
 	TArray<FAssetData> Assets;
 	AssetRegistryModule.Get().GetAssets(Filter, Assets);
@@ -122,7 +122,7 @@ bool FWebRemoteControlEditorRoutes::HandleGetThumbnailRoute(const FHttpServerReq
 			Builder.Append(LastSegment);
 
 			FARFilter FallbackFilter;
-			FallbackFilter.ObjectPaths.Add(Builder.ToString());
+			FallbackFilter.SoftObjectPaths.Add(FSoftObjectPath(Builder.ToView()));
 			AssetRegistryModule.Get().GetAssets(FallbackFilter, Assets);
 		}
 	}
diff --git a/Engine/Source/Developer/AssetTools/Private/AssetFixUpRedirectors.cpp b/Engine/Source/Developer/AssetTools/Private/AssetFixUpRedirectors.cpp
--- a/Engine/Source/Developer/AssetTools/Private/AssetFixUpRedirectors.cpp
+++ b/Engine/Source/Developer/AssetTools/Private/AssetFixUpRedirectors.cpp
@@ -390,7 +390,7 @@ void FAssetFixUpRedirectors::SaveReferencingCollections(TArray<FRedirectorRefs>&
 		// Follow each link in the redirector, and notify the collections manager that it is going to be removed - this will force it to re-save any required collections
 		for (UObjectRedirector* Redirector = RedirectorRefs.Redirector; Redirector; Redirector = Cast<UObjectRedirector>(Redirector->DestinationObject))
 		{
-			const FName RedirectorObjectPath = *Redirector->GetPathName();
+			const FSoftObjectPath RedirectorObjectPath = FSoftObjectPath(Redirector);
 			if (!CollectionManagerModule.Get().HandleRedirectorDeleted(RedirectorObjectPath))
 			{
 				RedirectorRefs.FailureReason = FText::Format(LOCTEXT("RedirectorFixupFailed_CollectionsFailedToSave", "Referencing collection(s) failed to save: {0}"), CollectionManagerModule.Get().GetLastError());
diff --git a/Engine/Source/Developer/AssetTools/Private/AssetRenameManager.cpp b/Engine/Source/Developer/AssetTools/Private/AssetRenameManager.cpp
--- a/Engine/Source/Developer/AssetTools/Private/AssetRenameManager.cpp
+++ b/Engine/Source/Developer/AssetTools/Private/AssetRenameManager.cpp
@@ -1163,7 +1163,7 @@ void FAssetRenameManager::DetectReferencingCollections(TArray<FAssetRenameDataWi
 		if (AssetToRename.Asset.IsValid())
 		{
 			TArray<FCollectionNameType> ReferencingCollections;
-			CollectionManagerModule.Get().GetCollectionsContainingObject(*AssetToRename.Asset->GetPathName(), ReferencingCollections);
+			CollectionManagerModule.Get().GetCollectionsContainingObject(FSoftObjectPath(AssetToRename.Asset.Get()), ReferencingCollections);
 
 			if (ReferencingCollections.Num() > 0)
 			{
diff --git a/Engine/Source/Developer/AssetTools/Private/CollectionAssetManagement.cpp b/Engine/Source/Developer/AssetTools/Private/CollectionAssetManagement.cpp
--- a/Engine/Source/Developer/AssetTools/Private/CollectionAssetManagement.cpp
+++ b/Engine/Source/Developer/AssetTools/Private/CollectionAssetManagement.cpp
@@ -18,8 +18,8 @@ FCollectionAssetManagement::FCollectionAssetManagement()
 	OnCollectionRenamedHandle = CollectionManagerModule.Get().OnCollectionRenamed().AddRaw(this, &FCollectionAssetManagement::HandleCollectionRenamed);
 	OnCollectionDestroyedHandle = CollectionManagerModule.Get().OnCollectionDestroyed().AddRaw(this, &FCollectionAssetManagement::HandleCollectionDestroyed);
 	OnCollectionUpdatedHandle = CollectionManagerModule.Get().OnCollectionUpdated().AddRaw(this, &FCollectionAssetManagement::HandleCollectionUpdated);
-	OnAssetsAddedHandle = CollectionManagerModule.Get().OnAssetsAdded().AddRaw(this, &FCollectionAssetManagement::HandleAssetsAddedToCollection);
-	OnAssetsRemovedHandle = CollectionManagerModule.Get().OnAssetsRemoved().AddRaw(this, &FCollectionAssetManagement::HandleAssetsRemovedFromCollection);
+	OnAssetsAddedHandle = CollectionManagerModule.Get().OnAssetsAddedToCollection().AddRaw(this, &FCollectionAssetManagement::HandleAssetsAddedToCollection);
+	OnAssetsRemovedHandle = CollectionManagerModule.Get().OnAssetsRemovedFromCollection().AddRaw(this, &FCollectionAssetManagement::HandleAssetsRemovedFromCollection);
 }
 
 FCollectionAssetManagement::~FCollectionAssetManagement()
@@ -32,8 +32,8 @@ FCollectionAssetManagement::~FCollectionAssetManagement()
 		CollectionManagerModule.Get().OnCollectionRenamed().Remove(OnCollectionRenamedHandle);
 		CollectionManagerModule.Get().OnCollectionDestroyed().Remove(OnCollectionDestroyedHandle);
 		CollectionManagerModule.Get().OnCollectionUpdated().Remove(OnCollectionUpdatedHandle);
-		CollectionManagerModule.Get().OnAssetsAdded().Remove(OnAssetsAddedHandle);
-		CollectionManagerModule.Get().OnAssetsRemoved().Remove(OnAssetsRemovedHandle);
+		CollectionManagerModule.Get().OnAssetsAddedToCollection().Remove(OnAssetsAddedHandle);
+		CollectionManagerModule.Get().OnAssetsRemovedFromCollection().Remove(OnAssetsRemovedHandle);
 	}
 }
 
@@ -60,9 +60,7 @@ void FCollectionAssetManagement::AddCurrentAssetsToCollection(FCollectionNameTyp
 {
 	FCollectionManagerModule& CollectionManagerModule = FCollectionManagerModule::GetModule();
 
-PRAGMA_DISABLE_DEPRECATION_WARNINGS
-	const TArray<FName> ObjectPaths = UE::SoftObjectPath::Private::ConvertSoftObjectPaths(CurrentAssetPaths.Array());
-PRAGMA_ENABLE_DEPRECATION_WARNINGS
+	const TArray<FSoftObjectPath> ObjectPaths = CurrentAssetPaths.Array();
 
 	FText ResultText;
 	bool bSuccess = false;
@@ -101,9 +99,7 @@ void FCollectionAssetManagement::RemoveCurrentAssetsFromCollection(FCollectionNa
 {
 	FCollectionManagerModule& CollectionManagerModule = FCollectionManagerModule::GetModule();
 
-PRAGMA_DISABLE_DEPRECATION_WARNINGS
-	const TArray<FName> ObjectPaths = UE::SoftObjectPath::Private::ConvertSoftObjectPaths(CurrentAssetPaths.Array());
-PRAGMA_ENABLE_DEPRECATION_WARNINGS
+	const TArray<FSoftObjectPath> ObjectPaths = CurrentAssetPaths.Array();
 
 	FText ResultText;
 	bool bSuccess = false;
@@ -173,9 +169,7 @@ void FCollectionAssetManagement::UpdateAssetManagementState()
 	if (CurrentAssetPaths.Num() == 1)
 	{
 		TArray<FCollectionNameType> MatchedCollections;
-PRAGMA_DISABLE_DEPRECATION_WARNINGS
-		CollectionManagerModule.Get().GetCollectionsContainingObject(CurrentAssetPaths.CreateConstIterator()->ToFName(), MatchedCollections, ECollectionRecursionFlags::Self);
-PRAGMA_ENABLE_DEPRECATION_WARNINGS
+		CollectionManagerModule.Get().GetCollectionsContainingObject(*CurrentAssetPaths.CreateConstIterator(), MatchedCollections, ECollectionRecursionFlags::Self);
 
 		for (const FCollectionNameType& CollectionKey : MatchedCollections)
 		{
@@ -183,9 +177,7 @@ PRAGMA_ENABLE_DEPRECATION_WARNINGS
 		}
 
 		MatchedCollections.Reset();
-PRAGMA_DISABLE_DEPRECATION_WARNINGS
-		CollectionManagerModule.Get().GetCollectionsContainingObject(CurrentAssetPaths.CreateConstIterator()->ToFName(), MatchedCollections, ECollectionRecursionFlags::Children);
-PRAGMA_ENABLE_DEPRECATION_WARNINGS
+		CollectionManagerModule.Get().GetCollectionsContainingObject(*CurrentAssetPaths.CreateConstIterator(), MatchedCollections, ECollectionRecursionFlags::Children);
 
 		for (const FCollectionNameType& CollectionKey : MatchedCollections)
 		{
@@ -197,17 +189,14 @@ PRAGMA_ENABLE_DEPRECATION_WARNINGS
 	}
 	else
 	{
-PRAGMA_DISABLE_DEPRECATION_WARNINGS
-		const TArray<FName> ObjectPaths = UE::SoftObjectPath::Private::ConvertSoftObjectPaths(CurrentAssetPaths.Array());
-PRAGMA_ENABLE_DEPRECATION_WARNINGS
-
-		TMap<FCollectionNameType, TArray<FName>> CollectionsAndMatchedObjects;
+		const TArray<FSoftObjectPath> ObjectPaths = CurrentAssetPaths.Array();
+		TMap<FCollectionNameType, TArray<FSoftObjectPath>> CollectionsAndMatchedObjects;
 		CollectionManagerModule.Get().GetCollectionsContainingObjects(ObjectPaths, CollectionsAndMatchedObjects);
 
-		for (const auto& MatchedCollection : CollectionsAndMatchedObjects)
+		for (const TPair<FCollectionNameType, TArray<FSoftObjectPath>>& MatchedCollection : CollectionsAndMatchedObjects)
 		{
 			const FCollectionNameType& CollectionKey = MatchedCollection.Key;
-			const TArray<FName>& MatchedObjects = MatchedCollection.Value;
+			const TArray<FSoftObjectPath>& MatchedObjects = MatchedCollection.Value;
 
 			// Collections that contain all of the selected assets are shown as checked, collections that only contain some of the selected assets are shown as undetermined
 			AssetManagementState.Add(
@@ -238,15 +227,13 @@ void FCollectionAssetManagement::HandleCollectionDestroyed(const FCollectionName
 	AssetManagementState.Remove(Collection);
 }
 
-void FCollectionAssetManagement::HandleAssetsAddedToCollection(const FCollectionNameType& Collection, const TArray<FName>& AssetsAdded)
+void FCollectionAssetManagement::HandleAssetsAddedToCollection(const FCollectionNameType& Collection, TConstArrayView<FSoftObjectPath> AssetsAdded)
 {
 	// Only need to update if one of the added assets belongs to our current selection set
 	bool bNeedsUpdate = false;
-	for (const FName& AssetPath : AssetsAdded)
+	for (const FSoftObjectPath& AssetPath : AssetsAdded)
 	{
-PRAGMA_DISABLE_DEPRECATION_WARNINGS
-		if (CurrentAssetPaths.Contains(FSoftObjectPath(AssetPath)))
-PRAGMA_ENABLE_DEPRECATION_WARNINGS
+		if (CurrentAssetPaths.Contains(AssetPath))
 		{
 			bNeedsUpdate = true;
 			break;
@@ -259,15 +246,13 @@ PRAGMA_ENABLE_DEPRECATION_WARNINGS
 	}
 }
 
-void FCollectionAssetManagement::HandleAssetsRemovedFromCollection(const FCollectionNameType& Collection, const TArray<FName>& AssetsRemoved)
+void FCollectionAssetManagement::HandleAssetsRemovedFromCollection(const FCollectionNameType& Collection, TConstArrayView<FSoftObjectPath> AssetsRemoved)
 {
 	// Only need to update if one of the removed assets belongs to our current selection set
 	bool bNeedsUpdate = false;
-	for (const FName& AssetPath : AssetsRemoved)
+	for (const FSoftObjectPath& AssetPath : AssetsRemoved)
 	{
-PRAGMA_DISABLE_DEPRECATION_WARNINGS
-		if (CurrentAssetPaths.Contains(FSoftObjectPath(AssetPath)))
-PRAGMA_ENABLE_DEPRECATION_WARNINGS
+		if (CurrentAssetPaths.Contains(AssetPath))
 		{
 			bNeedsUpdate = true;
 			break;
diff --git a/Engine/Source/Developer/AssetTools/Public/CollectionAssetManagement.h b/Engine/Source/Developer/AssetTools/Public/CollectionAssetManagement.h
--- a/Engine/Source/Developer/AssetTools/Public/CollectionAssetManagement.h
+++ b/Engine/Source/Developer/AssetTools/Public/CollectionAssetManagement.h
@@ -49,10 +49,10 @@ private:
 	void HandleCollectionDestroyed(const FCollectionNameType& Collection);
 
 	/** Handles assets being added to a collection */
-	void HandleAssetsAddedToCollection(const FCollectionNameType& Collection, const TArray<FName>& AssetsAdded);
+	void HandleAssetsAddedToCollection(const FCollectionNameType& Collection, TConstArrayView<FSoftObjectPath> AssetsAdded);
 
 	/** Handles assets being removed from a collection */
-	void HandleAssetsRemovedFromCollection(const FCollectionNameType& Collection, const TArray<FName>& AssetsRemoved);
+	void HandleAssetsRemovedFromCollection(const FCollectionNameType& Collection, TConstArrayView<FSoftObjectPath> AssetsRemoved);
 
 	/** Set of asset paths that we are currently observing and managing the collection state of */
 	TSet<FSoftObjectPath> CurrentAssetPaths;
diff --git a/Engine/Source/Developer/CollectionManager/Private/Collection.cpp b/Engine/Source/Developer/CollectionManager/Private/Collection.cpp
--- a/Engine/Source/Developer/CollectionManager/Private/Collection.cpp
+++ b/Engine/Source/Developer/CollectionManager/Private/Collection.cpp
@@ -22,10 +22,10 @@
 
 struct FCollectionUtils
 {
-	static void AppendCollectionToArray(const TSet<FName>& InObjectSet, TArray<FName>& OutObjectArray)
+	static void AppendCollectionToArray(const TSet<FSoftObjectPath>& InObjectSet, TArray<FSoftObjectPath>& OutObjectArray)
 	{
 		OutObjectArray.Reserve(OutObjectArray.Num() + InObjectSet.Num());
-		for (const FName& ObjectName : InObjectSet)
+		for (const FSoftObjectPath& ObjectName : InObjectSet)
 		{
 			OutObjectArray.Add(ObjectName);
 		}
@@ -127,26 +127,26 @@ bool FCollection::Load(FText& OutError)
 	{
 		const int32 NamesNum = FileContents.Num() - LineIndex;
 
-		TArray<FName> FNames;
-		FNames.SetNum(NamesNum);
+		TArray<FSoftObjectPath> Paths;
+		Paths.SetNum(NamesNum);
 
 		// Name hashing to register new FName takes time
 		// Process as much as possible in multiple threads
 		ParallelFor(
 			NamesNum,
-			[this, &FileContents, &LineIndex, &FNames](int32 LocalLineIndex)
+			[this, &FileContents, &LineIndex, &Paths](int32 LocalLineIndex)
 			{
 				FStringView Line(FileContents[LineIndex + LocalLineIndex]);
-				FNames[LocalLineIndex] = FName(Line.TrimStartAndEnd());
+				Paths[LocalLineIndex] = FSoftObjectPath(Line.TrimStartAndEnd());
 			},
 			// Do not pay for scheduling cost if number of items is too low
 			NamesNum < 1000 ? EParallelForFlags::ForceSingleThread : EParallelForFlags::None
 		);
 
 		// Static collection, a flat list of asset paths
-		for (FName& Name : FNames)
+		for (FSoftObjectPath& Path : Paths)
 		{
-			AddObjectToCollection(Name);
+			AddObjectToCollection(Path);
 		}
 	}
 	else
@@ -208,11 +208,11 @@ bool FCollection::Save(const TArray<FText>& AdditionalChangelistText, FText& Out
 	if (StorageMode == ECollectionStorageMode::Static)
 	{
 		// Write out the set as a sorted array to keep things in a known order for diffing
-		TArray<FName> ObjectList = ObjectSet.Array();
-		ObjectList.Sort(FNameLexicalLess());
+		TArray<FSoftObjectPath> ObjectList = ObjectSet.Array();
+		ObjectList.Sort([](FSoftObjectPath A, FSoftObjectPath B){ return A.LexicalLess(B); });
 
 		// Static collection. Save a flat list of all objects in the collection.
-		for (const FName& ObjectName : ObjectList)
+		for (const FSoftObjectPath& ObjectName : ObjectList)
 		{
 			FileOutput += ObjectName.ToString() + LINE_TERMINATOR;
 		}
@@ -411,9 +411,9 @@ void FCollection::Empty()
 	DiskSnapshot.TakeSnapshot(*this);
 }
 
-bool FCollection::AddObjectToCollection(FName ObjectPath)
+bool FCollection::AddObjectToCollection(const FSoftObjectPath& ObjectPath)
 {
-	if (ObjectPath.IsNone())
+	if (ObjectPath.IsNull())
 	{
 		return false;
 	}
@@ -428,9 +428,9 @@ bool FCollection::AddObjectToCollection(FName ObjectPath)
 	return false;
 }
 
-bool FCollection::RemoveObjectFromCollection(FName ObjectPath)
+bool FCollection::RemoveObjectFromCollection(const FSoftObjectPath& ObjectPath)
 {
-	if (ObjectPath.IsNone())
+	if (ObjectPath.IsNull())
 	{
 		return false;
 	}
@@ -443,13 +443,13 @@ bool FCollection::RemoveObjectFromCollection(FName ObjectPath)
 	return false;
 }
 
-void FCollection::GetAssetsInCollection(TArray<FName>& Assets) const
+void FCollection::GetAssetsInCollection(TArray<FSoftObjectPath>& Assets) const
 {
 	if (StorageMode == ECollectionStorageMode::Static)
 	{
-		for (const FName& ObjectName : ObjectSet)
+		for (const FSoftObjectPath& ObjectName : ObjectSet)
 		{
-			if (!ObjectName.ToString().StartsWith(TEXT("/Script/")))
+			if (!ObjectName.GetLongPackageName().StartsWith(TEXT("/Script/")))
 			{
 				Assets.Add(ObjectName);
 			}
@@ -457,21 +457,21 @@ void FCollection::GetAssetsInCollection(TArray<FName>& Assets) const
 	}
 }
 
-void FCollection::GetClassesInCollection(TArray<FName>& Classes) const
+void FCollection::GetClassesInCollection(TArray<FTopLevelAssetPath>& Classes) const
 {
 	if (StorageMode == ECollectionStorageMode::Static)
 	{
-		for (const FName& ObjectName : ObjectSet)
+		for (const FSoftObjectPath& ObjectName : ObjectSet)
 		{
-			if (ObjectName.ToString().StartsWith(TEXT("/Script/")))
+			if (ObjectName.GetLongPackageName().StartsWith(TEXT("/Script/")))
 			{
-				Classes.Add(ObjectName);
+				Classes.Add(ObjectName.GetAssetPath());
 			}
 		}
 	}
 }
 
-void FCollection::GetObjectsInCollection(TArray<FName>& Objects) const
+void FCollection::GetObjectsInCollection(TArray<FSoftObjectPath>& Objects) const
 {
 	if (StorageMode == ECollectionStorageMode::Static)
 	{
@@ -479,7 +479,7 @@ void FCollection::GetObjectsInCollection(TArray<FName>& Objects) const
 	}
 }
 
-bool FCollection::IsObjectInCollection(FName ObjectPath) const
+bool FCollection::IsObjectInCollection(const FSoftObjectPath& ObjectPath) const
 {
 	if (StorageMode == ECollectionStorageMode::Static)
 	{
@@ -489,7 +489,7 @@ bool FCollection::IsObjectInCollection(FName ObjectPath) const
 	return false;
 }
 
-bool FCollection::IsRedirectorInCollection(FName ObjectPath) const
+bool FCollection::IsRedirectorInCollection(const FSoftObjectPath& ObjectPath) const
 {
 	if (StorageMode == ECollectionStorageMode::Static)
 	{
@@ -576,8 +576,8 @@ bool FCollection::IsDirty() const
 
 	if (StorageMode == ECollectionStorageMode::Static)
 	{
-		TArray<FName> ObjectsAdded;
-		TArray<FName> ObjectsRemoved;
+		TArray<FSoftObjectPath> ObjectsAdded;
+		TArray<FSoftObjectPath> ObjectsRemoved;
 		GetObjectDifferencesFromDisk(ObjectsAdded, ObjectsRemoved);
 
 		bHasChanges = ObjectsAdded.Num() != 0 || ObjectsRemoved.Num() != 0;
@@ -610,10 +610,10 @@ void FCollection::PrintCollection() const
 		UE_LOG(LogCollectionManager, Log, TEXT("    ============================="));
 
 		// Print the set as a sorted array to keep things in a sane order
-		TArray<FName> ObjectList = ObjectSet.Array();
-		ObjectList.Sort(FNameLexicalLess());
+		TArray<FSoftObjectPath> ObjectList = ObjectSet.Array();
+		ObjectList.Sort([](const FSoftObjectPath& A, const FSoftObjectPath& B){ return A.LexicalLess(B); });
 
-		for (const FName& ObjectName : ObjectList)
+		for (const FSoftObjectPath& ObjectName : ObjectList)
 		{
 			UE_LOG(LogCollectionManager, Log, TEXT("        %s"), *ObjectName.ToString());
 		}
@@ -709,8 +709,8 @@ bool FCollection::MergeWithCollection(const FCollection& Other)
 	if (StorageMode == ECollectionStorageMode::Static)
 	{
 		// Work out whether we have any changes compared to the other collection
-		TArray<FName> ObjectsAdded;
-		TArray<FName> ObjectsRemoved;
+		TArray<FSoftObjectPath> ObjectsAdded;
+		TArray<FSoftObjectPath> ObjectsRemoved;
 		GetObjectDifferences(ObjectSet, Other.ObjectSet, ObjectsAdded, ObjectsRemoved);
 
 		bHasChanges = bHasChanges || ObjectsAdded.Num() > 0 || ObjectsRemoved.Num() > 0;
@@ -726,13 +726,13 @@ bool FCollection::MergeWithCollection(const FCollection& Other)
 			ObjectSet = Other.ObjectSet;
 
 			// Add the objects that were added before the merge
-			for (const FName& AddedObjectName : ObjectsAdded)
+			for (const FSoftObjectPath& AddedObjectName : ObjectsAdded)
 			{
 				ObjectSet.Add(AddedObjectName);
 			}
 
 			// Remove the objects that were removed before the merge
-			for (const FName& RemovedObjectName : ObjectsRemoved)
+			for (const FSoftObjectPath& RemovedObjectName : ObjectsRemoved)
 			{
 				ObjectSet.Remove(RemovedObjectName);
 			}
@@ -749,10 +749,10 @@ bool FCollection::MergeWithCollection(const FCollection& Other)
 	return bHasChanges;
 }
 
-void FCollection::GetObjectDifferences(const TSet<FName>& BaseSet, const TSet<FName>& NewSet, TArray<FName>& ObjectsAdded, TArray<FName>& ObjectsRemoved)
+void FCollection::GetObjectDifferences(const TSet<FSoftObjectPath>& BaseSet, const TSet<FSoftObjectPath>& NewSet, TArray<FSoftObjectPath>& ObjectsAdded, TArray<FSoftObjectPath>& ObjectsRemoved)
 {
 	// Find the objects that were removed compared to the base set
-	for (const FName& BaseObjectName : BaseSet)
+	for (const FSoftObjectPath& BaseObjectName : BaseSet)
 	{
 		if (!NewSet.Contains(BaseObjectName))
 		{
@@ -769,7 +769,7 @@ void FCollection::GetObjectDifferences(const TSet<FName>& BaseSet, const TSet<FN
 	}
 
 	// Find the objects that were added compare to the base set
-	for (const FName& NewObjectName : NewSet)
+	for (const FSoftObjectPath& NewObjectName : NewSet)
 	{
 		if (!BaseSet.Contains(NewObjectName))
 		{
@@ -778,7 +778,7 @@ void FCollection::GetObjectDifferences(const TSet<FName>& BaseSet, const TSet<FN
 	}
 }
 
-void FCollection::GetObjectDifferencesFromDisk(TArray<FName>& ObjectsAdded, TArray<FName>& ObjectsRemoved) const
+void FCollection::GetObjectDifferencesFromDisk(TArray<FSoftObjectPath>& ObjectsAdded, TArray<FSoftObjectPath>& ObjectsRemoved) const
 {
 	if (StorageMode == ECollectionStorageMode::Static)
 	{
@@ -952,18 +952,18 @@ bool FCollection::CheckinCollection(const TArray<FText>& AdditionalChangelistTex
 		if (StorageMode == ECollectionStorageMode::Static)
 		{
 			// Gather differences from disk
-			TArray<FName> ObjectsAdded;
-			TArray<FName> ObjectsRemoved;
+			TArray<FSoftObjectPath> ObjectsAdded;
+			TArray<FSoftObjectPath> ObjectsRemoved;
 			GetObjectDifferencesFromDisk(ObjectsAdded, ObjectsRemoved);
 
-			ObjectsAdded.Sort(FNameLexicalLess());
-			ObjectsRemoved.Sort(FNameLexicalLess());
+			ObjectsAdded.Sort([](FSoftObjectPath A, FSoftObjectPath B){ return A.LexicalLess(B); });
+			ObjectsRemoved.Sort([](FSoftObjectPath A, FSoftObjectPath B) { return A.LexicalLess(B); });
 
 			// Report added files
 			FFormatNamedArguments Args;
-			Args.Add(TEXT("FirstObjectAdded"), ObjectsAdded.Num() > 0 ? FText::FromName(ObjectsAdded[0]) : NSLOCTEXT("Core", "None", "None"));
+			Args.Add(TEXT("FirstObjectAdded"), ObjectsAdded.Num() > 0 ? FText::FromString(ObjectsAdded[0].ToString()) : NSLOCTEXT("Core", "None", "None"));
 			Args.Add(TEXT("NumberAdded"), FText::AsNumber(ObjectsAdded.Num()));
-			Args.Add(TEXT("FirstObjectRemoved"), ObjectsRemoved.Num() > 0 ? FText::FromName(ObjectsRemoved[0]) : NSLOCTEXT("Core", "None", "None"));
+			Args.Add(TEXT("FirstObjectRemoved"), ObjectsRemoved.Num() > 0 ? FText::FromString(ObjectsRemoved[0].ToString()) : NSLOCTEXT("Core", "None", "None"));
 			Args.Add(TEXT("NumberRemoved"), FText::AsNumber(ObjectsRemoved.Num()));
 			Args.Add(TEXT("CollectionName"), CollectionNameText);
 
@@ -976,9 +976,9 @@ bool FCollection::CheckinCollection(const TArray<FText>& AdditionalChangelistTex
 				ChangelistDescBuilder.AppendLineFormat(LOCTEXT("CollectionAddedMultipleDesc", "Added {NumberAdded} objects to collection '{CollectionName}':"), Args);
 
 				ChangelistDescBuilder.Indent();
-				for (const FName& AddedObjectName : ObjectsAdded)
+				for (const FSoftObjectPath& AddedObjectName : ObjectsAdded)
 				{
-					ChangelistDescBuilder.AppendLine(FText::FromName(AddedObjectName));
+					ChangelistDescBuilder.AppendLine(FText::FromString(AddedObjectName.ToString()));
 				}
 				ChangelistDescBuilder.Unindent();
 			}
@@ -992,9 +992,9 @@ bool FCollection::CheckinCollection(const TArray<FText>& AdditionalChangelistTex
 				ChangelistDescBuilder.AppendLineFormat(LOCTEXT("CollectionRemovedMultipleDesc", "Removed {NumberRemoved} objects from collection '{CollectionName}'"), Args);
 
 				ChangelistDescBuilder.Indent();
-				for (const FName& RemovedObjectName : ObjectsRemoved)
+				for (const FSoftObjectPath& RemovedObjectName : ObjectsRemoved)
 				{
-					ChangelistDescBuilder.AppendLine(FText::FromName(RemovedObjectName));
+					ChangelistDescBuilder.AppendLine(FText::FromString(RemovedObjectName.ToString()));
 				}
 				ChangelistDescBuilder.Unindent();
 			}
diff --git a/Engine/Source/Developer/CollectionManager/Private/Collection.h b/Engine/Source/Developer/CollectionManager/Private/Collection.h
--- a/Engine/Source/Developer/CollectionManager/Private/Collection.h
+++ b/Engine/Source/Developer/CollectionManager/Private/Collection.h
@@ -61,19 +61,19 @@ public:
 	void Empty();
 
 	/** Adds a single object to the collection. Static collections only. */
-	bool AddObjectToCollection(FName ObjectPath);
+	bool AddObjectToCollection(const FSoftObjectPath& ObjectPath);
 	/** Removes a single object from the collection. Static collections only. */
-	bool RemoveObjectFromCollection(FName ObjectPath);
+	bool RemoveObjectFromCollection(const FSoftObjectPath& ObjectPath);
 	/** Gets a list of assets in the collection. Static collections only. */
-	void GetAssetsInCollection(TArray<FName>& Assets) const;
-	/** Gets a list of classes in the collection. Static collections only. */
-	void GetClassesInCollection(TArray<FName>& Classes) const;
+	void GetAssetsInCollection(TArray<FSoftObjectPath>& Assets) const;
+	/** Gets a list of native classes in the collection. Static collections only. */
+	void GetClassesInCollection(TArray<FTopLevelAssetPath>& Classes) const;
 	/** Gets a list of objects in the collection. Static collections only. */
-	void GetObjectsInCollection(TArray<FName>& Objects) const;
+	void GetObjectsInCollection(TArray<FSoftObjectPath>& Objects) const;
 	/** Returns true when the specified object is in the collection. Static collections only. */
-	bool IsObjectInCollection(FName ObjectPath) const;
+	bool IsObjectInCollection(const FSoftObjectPath& ObjectPath) const;
 	/** Returns true when the specified redirector is in the collection. Static collections only. */
-	bool IsRedirectorInCollection(FName ObjectPath) const;
+	bool IsRedirectorInCollection(const FSoftObjectPath& ObjectPath) const;
 
 	/** Set the dynamic query text for this collection. Dynamic collections only. */
 	bool SetDynamicQueryText(const FString& InQueryText);
@@ -136,9 +136,9 @@ private:
 	/** Merges the assets from the specified collection with this collection */
 	bool MergeWithCollection(const FCollection& Other);
 	/** Gets the object differences between object set A (base) and B (new) */
-	static void GetObjectDifferences(const TSet<FName>& BaseSet, const TSet<FName>& NewSet, TArray<FName>& ObjectsAdded, TArray<FName>& ObjectsRemoved);
+	static void GetObjectDifferences(const TSet<FSoftObjectPath>& BaseSet, const TSet<FSoftObjectPath>& NewSet, TArray<FSoftObjectPath>& ObjectsAdded, TArray<FSoftObjectPath>& ObjectsRemoved);
 	/** Gets the object differences between what we have in memory, and what we loaded from disk. Static collections only. */
-	void GetObjectDifferencesFromDisk(TArray<FName>& ObjectsAdded, TArray<FName>& ObjectsRemoved) const;
+	void GetObjectDifferencesFromDisk(TArray<FSoftObjectPath>& ObjectsAdded, TArray<FSoftObjectPath>& ObjectsRemoved) const;
 	/** Checks the shared collection out from source control so it may be saved. If false, OutError is a human readable warning depicting the error. */
 	bool CheckoutCollection(FText& OutError);
 	/** Checks the shared collection in to source control after it is saved. In addition to the normal text, AdditionalChangelistText will be added to the checkin description. If false, OutError is a human readable warning depicting the error. */
@@ -167,7 +167,7 @@ private:
 		TOptional<FLinearColor> CollectionColor;
 
 		/** The set of objects in the collection. Takes the form PackageName.AssetName. Static collections only. */
-		TSet<FName> ObjectSet;
+		TSet<FSoftObjectPath> ObjectSet; // TODO: could be FTopLevelAssetPath
 
 		/** The dynamic query string for this collection. Dynamic collections only. */
 		FString DynamicQueryText;
@@ -192,7 +192,7 @@ private:
 	FString SourceFilename;
 
 	/** The set of objects in the collection. Takes the form PackageName.AssetName.Static collections only. */
-	TSet<FName> ObjectSet;
+	TSet<FSoftObjectPath> ObjectSet; // TODO: Could be FTopLevelAssetPath
 
 	/** The dynamic query string for this collection. Dynamic collections only. */
 	FString DynamicQueryText;
diff --git a/Engine/Source/Developer/CollectionManager/Private/CollectionManager.cpp b/Engine/Source/Developer/CollectionManager/Private/CollectionManager.cpp
--- a/Engine/Source/Developer/CollectionManager/Private/CollectionManager.cpp
+++ b/Engine/Source/Developer/CollectionManager/Private/CollectionManager.cpp
@@ -80,7 +80,7 @@ const FCollectionObjectsMap& FCollectionManagerCache::GetCachedObjects() const
 			const FCollectionNameType& CollectionKey = AvailableCollection.Key;
 			const TSharedRef<FCollection>& Collection = AvailableCollection.Value;
 
-			TArray<FName> ObjectsInCollection;
+			TArray<FSoftObjectPath> ObjectsInCollection;
 			Collection->GetObjectsInCollection(ObjectsInCollection);
 
 			if (ObjectsInCollection.Num() > 0)
@@ -104,9 +104,9 @@ const FCollectionObjectsMap& FCollectionManagerCache::GetCachedObjects() const
 						break;
 					}
 
-					for (const FName& ObjectPath : ObjectsInCollection)
+					for (const FSoftObjectPath& ObjectPath : ObjectsInCollection)
 					{
-						auto& ObjectCollectionInfos = CachedObjects_Internal.FindOrAdd(ObjectPath);
+						TArray<FObjectCollectionInfo>& ObjectCollectionInfos = CachedObjects_Internal.FindOrAdd(ObjectPath);
 						FObjectCollectionInfo* ObjectInfoPtr = ObjectCollectionInfos.FindByPredicate([&](const FObjectCollectionInfo& InCollectionInfo) { return InCollectionInfo.CollectionKey == InCollectionKey; });
 						if (ObjectInfoPtr)
 						{
@@ -300,10 +300,12 @@ FCollectionManager::FCollectionManager()
 	TickFileCacheDelegateHandle = FTSTicker::GetCoreTicker().AddTicker(FTickerDelegate::CreateRaw(this, &FCollectionManager::TickFileCache), 1.0f);
 }
 
+PRAGMA_DISABLE_DEPRECATION_WARNINGS
 FCollectionManager::~FCollectionManager()
 {
 	FTSTicker::GetCoreTicker().RemoveTicker(TickFileCacheDelegateHandle);
 }
+PRAGMA_ENABLE_DEPRECATION_WARNINGS
 
 bool FCollectionManager::HasCollections() const
 {
@@ -499,7 +501,7 @@ bool FCollectionManager::CollectionExists(FName CollectionName, ECollectionShare
 	}
 }
 
-bool FCollectionManager::GetAssetsInCollection(FName CollectionName, ECollectionShareType::Type ShareType, TArray<FName>& AssetsPaths, ECollectionRecursionFlags::Flags RecursionMode) const
+bool FCollectionManager::GetAssetsInCollection(FName CollectionName, ECollectionShareType::Type ShareType, TArray<FSoftObjectPath>& AssetsPaths, ECollectionRecursionFlags::Flags RecursionMode) const
 {
 	bool bFoundAssets = false;
 
@@ -530,7 +532,7 @@ bool FCollectionManager::GetAssetsInCollection(FName CollectionName, ECollection
 	return bFoundAssets;
 }
 
-bool FCollectionManager::GetClassesInCollection(FName CollectionName, ECollectionShareType::Type ShareType, TArray<FName>& ClassPaths, ECollectionRecursionFlags::Flags RecursionMode) const
+bool FCollectionManager::GetClassesInCollection(FName CollectionName, ECollectionShareType::Type ShareType, TArray<FTopLevelAssetPath>& ClassPaths, ECollectionRecursionFlags::Flags RecursionMode) const
 {
 	bool bFoundClasses = false;
 
@@ -561,7 +563,7 @@ bool FCollectionManager::GetClassesInCollection(FName CollectionName, ECollectio
 	return bFoundClasses;
 }
 
-bool FCollectionManager::GetObjectsInCollection(FName CollectionName, ECollectionShareType::Type ShareType, TArray<FName>& ObjectPaths, ECollectionRecursionFlags::Flags RecursionMode) const
+bool FCollectionManager::GetObjectsInCollection(FName CollectionName, ECollectionShareType::Type ShareType, TArray<FSoftObjectPath>& ObjectPaths, ECollectionRecursionFlags::Flags RecursionMode) const
 {
 	bool bFoundObjects = false;
 
@@ -592,7 +594,7 @@ bool FCollectionManager::GetObjectsInCollection(FName CollectionName, ECollectio
 	return bFoundObjects;
 }
 
-void FCollectionManager::GetCollectionsContainingObject(FName ObjectPath, ECollectionShareType::Type ShareType, TArray<FName>& OutCollectionNames, ECollectionRecursionFlags::Flags RecursionMode) const
+void FCollectionManager::GetCollectionsContainingObject(const FSoftObjectPath& ObjectPath, ECollectionShareType::Type ShareType, TArray<FName>& OutCollectionNames, ECollectionRecursionFlags::Flags RecursionMode) const
 {
 	const FCollectionObjectsMap& CachedObjects = CollectionCache.GetCachedObjects();
 
@@ -609,7 +611,7 @@ void FCollectionManager::GetCollectionsContainingObject(FName ObjectPath, EColle
 	}
 }
 
-void FCollectionManager::GetCollectionsContainingObject(FName ObjectPath, TArray<FCollectionNameType>& OutCollections, ECollectionRecursionFlags::Flags RecursionMode) const
+void FCollectionManager::GetCollectionsContainingObject(const FSoftObjectPath& ObjectPath, TArray<FCollectionNameType>& OutCollections, ECollectionRecursionFlags::Flags RecursionMode) const
 {
 	const FCollectionObjectsMap& CachedObjects = CollectionCache.GetCachedObjects();
 
@@ -627,11 +629,11 @@ void FCollectionManager::GetCollectionsContainingObject(FName ObjectPath, TArray
 	}
 }
 
-void FCollectionManager::GetCollectionsContainingObjects(const TArray<FName>& ObjectPaths, TMap<FCollectionNameType, TArray<FName>>& OutCollectionsAndMatchedObjects, ECollectionRecursionFlags::Flags RecursionMode) const
+void FCollectionManager::GetCollectionsContainingObjects(const TArray<FSoftObjectPath>& ObjectPaths, TMap<FCollectionNameType, TArray<FSoftObjectPath>>& OutCollectionsAndMatchedObjects, ECollectionRecursionFlags::Flags RecursionMode) const
 {
 	const FCollectionObjectsMap& CachedObjects = CollectionCache.GetCachedObjects();
 
-	for (const FName& ObjectPath : ObjectPaths)
+	for (const FSoftObjectPath& ObjectPath : ObjectPaths)
 	{
 		const auto* ObjectCollectionInfosPtr = CachedObjects.Find(ObjectPath);
 		if (ObjectCollectionInfosPtr)
@@ -640,7 +642,7 @@ void FCollectionManager::GetCollectionsContainingObjects(const TArray<FName>& Ob
 			{
 				if ((RecursionMode & ObjectCollectionInfo.Reason) != 0)
 				{
-					TArray<FName>& MatchedObjects = OutCollectionsAndMatchedObjects.FindOrAdd(ObjectCollectionInfo.CollectionKey);
+					TArray<FSoftObjectPath>& MatchedObjects = OutCollectionsAndMatchedObjects.FindOrAdd(ObjectCollectionInfo.CollectionKey);
 					MatchedObjects.Add(ObjectPath);
 				}
 			}
@@ -648,7 +650,7 @@ void FCollectionManager::GetCollectionsContainingObjects(const TArray<FName>& Ob
 	}
 }
 
-FString FCollectionManager::GetCollectionsStringForObject(FName ObjectPath, ECollectionShareType::Type ShareType, ECollectionRecursionFlags::Flags RecursionMode, bool bFullPaths) const
+FString FCollectionManager::GetCollectionsStringForObject(const FSoftObjectPath& ObjectPath, ECollectionShareType::Type ShareType, ECollectionRecursionFlags::Flags RecursionMode, bool bFullPaths) const
 {
 	const FCollectionObjectsMap& CachedObjects = CollectionCache.GetCachedObjects();
 
@@ -991,14 +993,14 @@ bool FCollectionManager::DestroyCollection(FName CollectionName, ECollectionShar
 	}
 }
 
-bool FCollectionManager::AddToCollection(FName CollectionName, ECollectionShareType::Type ShareType, FName ObjectPath)
+bool FCollectionManager::AddToCollection(FName CollectionName, ECollectionShareType::Type ShareType, const FSoftObjectPath& ObjectPath)
 {
-	TArray<FName> Paths;
+	TArray<FSoftObjectPath> Paths;
 	Paths.Add(ObjectPath);
 	return AddToCollection(CollectionName, ShareType, Paths);
 }
 
-bool FCollectionManager::AddToCollection(FName CollectionName, ECollectionShareType::Type ShareType, const TArray<FName>& ObjectPaths, int32* OutNumAdded)
+bool FCollectionManager::AddToCollection(FName CollectionName, ECollectionShareType::Type ShareType, TConstArrayView<FSoftObjectPath> ObjectPaths, int32* OutNumAdded)
 {
 	if (OutNumAdded)
 	{
@@ -1028,7 +1030,7 @@ bool FCollectionManager::AddToCollection(FName CollectionName, ECollectionShareT
 	}
 
 	int32 NumAdded = 0;
-	for (const FName& ObjectPath : ObjectPaths)
+	for (const FSoftObjectPath& ObjectPath : ObjectPaths)
 	{
 		if ((*CollectionRefPtr)->AddObjectToCollection(ObjectPath))
 		{
@@ -1049,13 +1051,19 @@ bool FCollectionManager::AddToCollection(FName CollectionName, ECollectionShareT
 			}
 
 			CollectionCache.HandleCollectionChanged();
-			AssetsAddedEvent.Broadcast(CollectionKey, ObjectPaths);
+			AssetsAddedToCollectionDelegate.Broadcast(CollectionKey, ObjectPaths);
+PRAGMA_DISABLE_DEPRECATION_WARNINGS
+			if (AssetsAddedEvent.IsBound()) // Avoid conversion if nothing is bound 
+			{
+				AssetsAddedEvent.Broadcast(CollectionKey, UE::SoftObjectPath::Private::ConvertSoftObjectPaths(ObjectPaths));
+			}
+PRAGMA_ENABLE_DEPRECATION_WARNINGS
 			return true;
 		}
 		else
 		{
 			// Added but not saved, revert the add
-			for (const FName& ObjectPath : ObjectPaths)
+			for (const FSoftObjectPath& ObjectPath : ObjectPaths)
 			{
 				(*CollectionRefPtr)->RemoveObjectFromCollection(ObjectPath);
 			}
@@ -1070,14 +1078,14 @@ bool FCollectionManager::AddToCollection(FName CollectionName, ECollectionShareT
 	}
 }
 
-bool FCollectionManager::RemoveFromCollection(FName CollectionName, ECollectionShareType::Type ShareType, FName ObjectPath)
+bool FCollectionManager::RemoveFromCollection(FName CollectionName, ECollectionShareType::Type ShareType, const FSoftObjectPath& ObjectPath)
 {
-	TArray<FName> Paths;
+	TArray<FSoftObjectPath> Paths;
 	Paths.Add(ObjectPath);
 	return RemoveFromCollection(CollectionName, ShareType, Paths);
 }
 
-bool FCollectionManager::RemoveFromCollection(FName CollectionName, ECollectionShareType::Type ShareType, const TArray<FName>& ObjectPaths, int32* OutNumRemoved)
+bool FCollectionManager::RemoveFromCollection(FName CollectionName, ECollectionShareType::Type ShareType, TConstArrayView<FSoftObjectPath> ObjectPaths, int32* OutNumRemoved)
 {
 	if (OutNumRemoved)
 	{
@@ -1106,8 +1114,8 @@ bool FCollectionManager::RemoveFromCollection(FName CollectionName, ECollectionS
 		return false;
 	}
 
-	TArray<FName> RemovedAssets;
-	for (const FName& ObjectPath : ObjectPaths)
+	TArray<FSoftObjectPath> RemovedAssets;
+	for (const FSoftObjectPath& ObjectPath : ObjectPaths)
 	{
 		if ((*CollectionRefPtr)->RemoveObjectFromCollection(ObjectPath))
 		{
@@ -1133,13 +1141,19 @@ bool FCollectionManager::RemoveFromCollection(FName CollectionName, ECollectionS
 		}
 
 		CollectionCache.HandleCollectionChanged();
-		AssetsRemovedEvent.Broadcast(CollectionKey, ObjectPaths);
+		AssetsRemovedFromCollectionDelegate.Broadcast(CollectionKey, ObjectPaths);
+PRAGMA_DISABLE_DEPRECATION_WARNINGS
+		if (AssetsRemovedEvent.IsBound())
+		{
+			AssetsRemovedEvent.Broadcast(CollectionKey, UE::SoftObjectPath::Private::ConvertSoftObjectPaths(ObjectPaths));
+		}
+PRAGMA_ENABLE_DEPRECATION_WARNINGS
 		return true;
 	}
 	else
 	{
 		// Removed but not saved, revert the remove
-		for (const FName& RemovedAssetName : RemovedAssets)
+		for (const FSoftObjectPath& RemovedAssetName : RemovedAssets)
 		{
 			(*CollectionRefPtr)->AddObjectToCollection(RemovedAssetName);
 		}
@@ -1462,7 +1476,7 @@ bool FCollectionManager::GetCollectionStorageMode(FName CollectionName, ECollect
 	return false;
 }
 
-bool FCollectionManager::IsObjectInCollection(FName ObjectPath, FName CollectionName, ECollectionShareType::Type ShareType, ECollectionRecursionFlags::Flags RecursionMode) const
+bool FCollectionManager::IsObjectInCollection(const FSoftObjectPath& ObjectPath, FName CollectionName, ECollectionShareType::Type ShareType, ECollectionRecursionFlags::Flags RecursionMode) const
 {
 	if (!ensure(ShareType < ECollectionShareType::CST_All))
 	{
@@ -1554,14 +1568,14 @@ void FCollectionManager::HandleFixupRedirectors(ICollectionRedirectorFollower& I
 
 	const double LoadStartTime = FPlatformTime::Seconds();
 
-	TArray<TPair<FName, FName>> ObjectsToRename;
+	TArray<TPair<FSoftObjectPath, FSoftObjectPath>> ObjectsToRename;
 
 	// Build up the list of redirected object into rename pairs
 	{
 		const FCollectionObjectsMap& CachedObjects = CollectionCache.GetCachedObjects();
 		for (const auto& CachedObjectInfo : CachedObjects)
 		{
-			FName NewObjectPath;
+			FSoftObjectPath NewObjectPath;
 			if (InRedirectorFollower.FixupObject(CachedObjectInfo.Key, NewObjectPath))
 			{
 				ObjectsToRename.Emplace(CachedObjectInfo.Key, NewObjectPath);
@@ -1571,14 +1585,14 @@ void FCollectionManager::HandleFixupRedirectors(ICollectionRedirectorFollower& I
 
 	TArray<FCollectionNameType> UpdatedCollections;
 
-	TArray<FName> AddedObjects;
+	TArray<FSoftObjectPath> AddedObjects;
 	AddedObjects.Reserve(ObjectsToRename.Num());
 
-	TArray<FName> RemovedObjects;
+	TArray<FSoftObjectPath> RemovedObjects;
 	RemovedObjects.Reserve(ObjectsToRename.Num());
 
 	// Handle the rename for each redirected object
-	for (const auto& ObjectToRename : ObjectsToRename)
+	for (const TPair<FSoftObjectPath, FSoftObjectPath>& ObjectToRename : ObjectsToRename)
 	{
 		AddedObjects.Add(ObjectToRename.Value);
 		RemovedObjects.Add(ObjectToRename.Key);
@@ -1593,9 +1607,22 @@ void FCollectionManager::HandleFixupRedirectors(ICollectionRedirectorFollower& I
 		// Notify every collection that changed
 		for (const FCollectionNameType& UpdatedCollection : UpdatedCollections)
 		{
-			AssetsRemovedEvent.Broadcast(UpdatedCollection, RemovedObjects);
-			AssetsAddedEvent.Broadcast(UpdatedCollection, AddedObjects);
+			AssetsRemovedFromCollectionDelegate.Broadcast(UpdatedCollection, RemovedObjects);
+			AssetsAddedToCollectionDelegate.Broadcast(UpdatedCollection, AddedObjects);
+		}
+
+PRAGMA_DISABLE_DEPRECATION_WARNINGS
+		if (AssetsRemovedEvent.IsBound() || AssetsAddedEvent.IsBound())
+		{
+			TArray<FName> RemovedObjectPathNames = UE::SoftObjectPath::Private::ConvertSoftObjectPaths(RemovedObjects);
+			TArray<FName> AddedObjectPathNames = UE::SoftObjectPath::Private::ConvertSoftObjectPaths(AddedObjects);
+			for (const FCollectionNameType& UpdatedCollection : UpdatedCollections)
+			{
+				AssetsRemovedEvent.Broadcast(UpdatedCollection, RemovedObjectPathNames);
+				AssetsAddedEvent.Broadcast(UpdatedCollection, AddedObjectPathNames);
+			}
 		}
+PRAGMA_ENABLE_DEPRECATION_WARNINGS
 	}
 
 	UE_LOG(LogCollectionManager, Log, TEXT( "Fixed up redirectors for %d collections in %0.6f seconds (updated %d objects)" ), AvailableCollections.Num(), FPlatformTime::Seconds() - LoadStartTime, ObjectsToRename.Num());
@@ -1606,7 +1633,7 @@ void FCollectionManager::HandleFixupRedirectors(ICollectionRedirectorFollower& I
 	}
 }
 
-bool FCollectionManager::HandleRedirectorDeleted(const FName& ObjectPath)
+bool FCollectionManager::HandleRedirectorDeleted(const FSoftObjectPath& ObjectPath)
 {
 	bool bSavedAllCollections = true;
 
@@ -1637,15 +1664,26 @@ bool FCollectionManager::HandleRedirectorDeleted(const FName& ObjectPath)
 		}
 	}
 
-	TArray<FName> RemovedObjects;
+	TArray<FSoftObjectPath> RemovedObjects;
 	RemovedObjects.Add(ObjectPath);
 
 	// Notify every collection that changed
 	for (const FCollectionNameType& UpdatedCollection : UpdatedCollections)
 	{
-		AssetsRemovedEvent.Broadcast(UpdatedCollection, RemovedObjects);
+		AssetsRemovedFromCollectionDelegate.Broadcast(UpdatedCollection, RemovedObjects);
 	}
 
+PRAGMA_DISABLE_DEPRECATION_WARNINGS
+	if( AssetsRemovedEvent.IsBound())
+	{
+		TArray<FName> RemovedObjectPathNames = UE::SoftObjectPath::Private::ConvertSoftObjectPaths(RemovedObjects);
+		for (const FCollectionNameType& UpdatedCollection : UpdatedCollections)
+		{
+			AssetsRemovedEvent.Broadcast(UpdatedCollection, RemovedObjectPathNames);
+		}
+	}
+PRAGMA_ENABLE_DEPRECATION_WARNINGS
+
 	if (!bSavedAllCollections)
 	{
 		LastError = AllErrors.ToText();
@@ -1654,15 +1692,15 @@ bool FCollectionManager::HandleRedirectorDeleted(const FName& ObjectPath)
 	return bSavedAllCollections;
 }
 
-void FCollectionManager::HandleObjectRenamed(const FName& OldObjectPath, const FName& NewObjectPath)
+void FCollectionManager::HandleObjectRenamed(const FSoftObjectPath& OldObjectPath, const FSoftObjectPath& NewObjectPath)
 {
-	TArray<FCollectionNameType> UpdatedCollections;
+	TArray<FCollectionNameType> UpdatedCollections; 
 	ReplaceObjectInCollections(OldObjectPath, NewObjectPath, UpdatedCollections);
 
-	TArray<FName> AddedObjects;
+	TArray<FSoftObjectPath> AddedObjects;
 	AddedObjects.Add(NewObjectPath);
 
-	TArray<FName> RemovedObjects;
+	TArray<FSoftObjectPath> RemovedObjects;
 	RemovedObjects.Add(OldObjectPath);
 
 	if (UpdatedCollections.Num() > 0)
@@ -1672,18 +1710,31 @@ void FCollectionManager::HandleObjectRenamed(const FName& OldObjectPath, const F
 		// Notify every collection that changed
 		for (const FCollectionNameType& UpdatedCollection : UpdatedCollections)
 		{
-			AssetsRemovedEvent.Broadcast(UpdatedCollection, RemovedObjects);
-			AssetsAddedEvent.Broadcast(UpdatedCollection, AddedObjects);
+			AssetsRemovedFromCollectionDelegate.Broadcast(UpdatedCollection, RemovedObjects);
+			AssetsAddedToCollectionDelegate.Broadcast(UpdatedCollection, AddedObjects);
 		}
+
+PRAGMA_DISABLE_DEPRECATION_WARNINGS
+		if( AssetsAddedEvent.IsBound() || AssetsRemovedEvent.IsBound())
+		{
+			TArray<FName> RemovedObjectPathNames = UE::SoftObjectPath::Private::ConvertSoftObjectPaths(RemovedObjects);
+			TArray<FName> AddedObjectPathNames = UE::SoftObjectPath::Private::ConvertSoftObjectPaths(AddedObjects);
+			for (const FCollectionNameType& UpdatedCollection : UpdatedCollections)
+			{
+				AssetsRemovedEvent.Broadcast(UpdatedCollection, RemovedObjectPathNames);
+				AssetsAddedEvent.Broadcast(UpdatedCollection, AddedObjectPathNames);
+			}
+		}
+PRAGMA_ENABLE_DEPRECATION_WARNINGS
 	}
 }
 
-void FCollectionManager::HandleObjectDeleted(const FName& ObjectPath)
+void FCollectionManager::HandleObjectDeleted(const FSoftObjectPath& ObjectPath)
 {
 	TArray<FCollectionNameType> UpdatedCollections;
 	RemoveObjectFromCollections(ObjectPath, UpdatedCollections);
 
-	TArray<FName> RemovedObjects;
+	TArray<FSoftObjectPath> RemovedObjects;
 	RemovedObjects.Add(ObjectPath);
 
 	if (UpdatedCollections.Num() > 0)
@@ -1693,8 +1744,19 @@ void FCollectionManager::HandleObjectDeleted(const FName& ObjectPath)
 		// Notify every collection that changed
 		for (const FCollectionNameType& UpdatedCollection : UpdatedCollections)
 		{
-			AssetsRemovedEvent.Broadcast(UpdatedCollection, RemovedObjects);
+			AssetsRemovedFromCollectionDelegate.Broadcast(UpdatedCollection, RemovedObjects);
+		}
+
+PRAGMA_DISABLE_DEPRECATION_WARNINGS
+		if (AssetsRemovedEvent.IsBound())
+		{
+			TArray<FName> RemovedObjectPathNames = UE::SoftObjectPath::Private::ConvertSoftObjectPaths(RemovedObjects);
+			for (const FCollectionNameType& UpdatedCollection : UpdatedCollections)
+			{
+				AssetsRemovedEvent.Broadcast(UpdatedCollection, RemovedObjectPathNames);
+			}
 		}
+PRAGMA_ENABLE_DEPRECATION_WARNINGS
 	}
 }
 
@@ -1929,7 +1991,7 @@ bool FCollectionManager::RemoveCollection(const TSharedRef<FCollection>& Collect
 	return false;
 }
 
-void FCollectionManager::RemoveObjectFromCollections(const FName& ObjectPath, TArray<FCollectionNameType>& OutUpdatedCollections)
+void FCollectionManager::RemoveObjectFromCollections(const FSoftObjectPath& ObjectPath, TArray<FCollectionNameType>& OutUpdatedCollections)
 {
 	const FCollectionObjectsMap& CachedObjects = CollectionCache.GetCachedObjects();
 
@@ -1956,7 +2018,7 @@ void FCollectionManager::RemoveObjectFromCollections(const FName& ObjectPath, TA
 	}
 }
 
-void FCollectionManager::ReplaceObjectInCollections(const FName& OldObjectPath, const FName& NewObjectPath, TArray<FCollectionNameType>& OutUpdatedCollections)
+void FCollectionManager::ReplaceObjectInCollections(const FSoftObjectPath& OldObjectPath, const FSoftObjectPath& NewObjectPath, TArray<FCollectionNameType>& OutUpdatedCollections)
 {
 	const FCollectionObjectsMap& CachedObjects = CollectionCache.GetCachedObjects();
 
diff --git a/Engine/Source/Developer/CollectionManager/Private/CollectionManager.h b/Engine/Source/Developer/CollectionManager/Private/CollectionManager.h
--- a/Engine/Source/Developer/CollectionManager/Private/CollectionManager.h
+++ b/Engine/Source/Developer/CollectionManager/Private/CollectionManager.h
@@ -36,7 +36,7 @@ struct FObjectCollectionInfo
 
 typedef TMap<FCollectionNameType, TSharedRef<FCollection>> FAvailableCollectionsMap;
 typedef TMap<FGuid, FCollectionNameType> FGuidToCollectionNamesMap;
-typedef TMap<FName, TArray<FObjectCollectionInfo>> FCollectionObjectsMap;
+typedef TMap<FSoftObjectPath, TArray<FObjectCollectionInfo>> FCollectionObjectsMap;
 typedef TMap<FGuid, TArray<FGuid>> FCollectionHierarchyMap;
 typedef TArray<FLinearColor> FCollectionColorArray;
 
@@ -115,6 +115,103 @@ public:
 	virtual ~FCollectionManager();
 
 	// ICollectionManager implementation
+PRAGMA_DISABLE_DEPRECATION_WARNINGS
+	virtual bool GetAssetsInCollection(FName CollectionName, ECollectionShareType::Type ShareType, TArray<FName>& AssetPaths, ECollectionRecursionFlags::Flags RecursionMode = ECollectionRecursionFlags::Self) const override
+	{
+		TArray<FSoftObjectPath> Temp;
+		if (GetAssetsInCollection(CollectionName, ShareType, Temp, RecursionMode))
+		{
+			AssetPaths.Append(UE::SoftObjectPath::Private::ConvertSoftObjectPaths(Temp));
+			return true;
+		}
+		return false;
+	}
+	virtual bool GetObjectsInCollection(FName CollectionName, ECollectionShareType::Type ShareType, TArray<FName>& ObjectPaths, ECollectionRecursionFlags::Flags RecursionMode = ECollectionRecursionFlags::Self) const override
+	{
+		TArray<FSoftObjectPath> Temp;
+		if (GetObjectsInCollection(CollectionName, ShareType, Temp, RecursionMode))
+		{
+			ObjectPaths.Append(UE::SoftObjectPath::Private::ConvertSoftObjectPaths(Temp));
+			return true;
+		}
+		return false;
+	}
+	virtual bool GetClassesInCollection(FName CollectionName, ECollectionShareType::Type ShareType, TArray<FName>& ClassPaths, ECollectionRecursionFlags::Flags RecursionMode = ECollectionRecursionFlags::Self) const override
+	{
+		TArray<FTopLevelAssetPath> Temp;
+		if (GetClassesInCollection(CollectionName, ShareType, Temp, RecursionMode))
+		{
+			for (FTopLevelAssetPath Path : Temp)
+			{
+				ClassPaths.Add(Path.ToFName());
+			}
+		
+			return true;
+		}
+		return false;
+	}
+	virtual void GetCollectionsContainingObject(FName ObjectPath, ECollectionShareType::Type ShareType, TArray<FName>& OutCollectionNames, ECollectionRecursionFlags::Flags RecursionMode = ECollectionRecursionFlags::Self) const override
+	{
+		GetCollectionsContainingObject(FSoftObjectPath(ObjectPath), ShareType, OutCollectionNames, RecursionMode);
+	}
+	virtual void GetCollectionsContainingObject(FName ObjectPath, TArray<FCollectionNameType>& OutCollections, ECollectionRecursionFlags::Flags RecursionMode = ECollectionRecursionFlags::Self) const override
+	{
+		GetCollectionsContainingObject(FSoftObjectPath(ObjectPath), OutCollections, RecursionMode);
+	}
+	virtual void GetCollectionsContainingObjects(const TArray<FName>& ObjectPathNames, TMap<FCollectionNameType, TArray<FName>>& OutCollectionsAndMatchedObjects, ECollectionRecursionFlags::Flags RecursionMode = ECollectionRecursionFlags::Self) const override
+	{
+		TArray<FSoftObjectPath> Paths = UE::SoftObjectPath::Private::ConvertObjectPathNames(ObjectPathNames);
+		TMap<FCollectionNameType, TArray<FSoftObjectPath>> TmpMap;
+		GetCollectionsContainingObjects(Paths, TmpMap, RecursionMode);
+		for (const TPair<FCollectionNameType, TArray<FSoftObjectPath>>& Pair : TmpMap)
+		{
+			TArray<FName>& Names = OutCollectionsAndMatchedObjects.FindOrAdd(Pair.Key);
+			Names.Append(UE::SoftObjectPath::Private::ConvertSoftObjectPaths(Pair.Value));
+		}
+	}
+	virtual FString GetCollectionsStringForObject(FName ObjectPath, ECollectionShareType::Type ShareType, ECollectionRecursionFlags::Flags RecursionMode = ECollectionRecursionFlags::Self, bool bFullPaths = true) const override
+	{
+		return GetCollectionsStringForObject(FSoftObjectPath(ObjectPath), ShareType, RecursionMode, bFullPaths);
+	}
+	virtual bool AddToCollection(FName CollectionName, ECollectionShareType::Type ShareType, FName ObjectPath) override
+	{
+		return AddToCollection(CollectionName, ShareType, FSoftObjectPath(ObjectPath));
+	}
+	virtual bool AddToCollection(FName CollectionName, ECollectionShareType::Type ShareType, const TArray<FName>& ObjectPaths, int32* OutNumAdded = nullptr) override
+	{
+		return AddToCollection(CollectionName, ShareType, UE::SoftObjectPath::Private::ConvertObjectPathNames(ObjectPaths), OutNumAdded);
+	}
+	virtual bool RemoveFromCollection(FName CollectionName, ECollectionShareType::Type ShareType, FName ObjectPath) override
+	{
+		return RemoveFromCollection(CollectionName, ShareType, FSoftObjectPath(ObjectPath));
+	}
+	virtual bool RemoveFromCollection(FName CollectionName, ECollectionShareType::Type ShareType, const TArray<FName>& ObjectPaths, int32* OutNumRemoved = nullptr) override
+	{
+		return RemoveFromCollection(CollectionName, ShareType, UE::SoftObjectPath::Private::ConvertObjectPathNames(ObjectPaths), OutNumRemoved);
+	}
+	virtual bool IsObjectInCollection(FName ObjectPath, FName CollectionName, ECollectionShareType::Type ShareType, ECollectionRecursionFlags::Flags RecursionMode = ECollectionRecursionFlags::Self) const override
+	{
+		return IsObjectInCollection(FSoftObjectPath(ObjectPath), CollectionName, ShareType, RecursionMode);
+	}
+	virtual bool HandleRedirectorDeleted(const FName& ObjectPath) override
+	{
+		return HandleRedirectorDeleted(FSoftObjectPath(ObjectPath));
+	}
+	virtual void HandleObjectRenamed(const FName& OldObjectPath, const FName& NewObjectPath) override
+	{
+		return HandleObjectRenamed(FSoftObjectPath(OldObjectPath), FSoftObjectPath(NewObjectPath));
+	}
+	virtual void HandleObjectDeleted(const FName& ObjectPath) override
+	{
+		return HandleObjectDeleted(FSoftObjectPath(ObjectPath));
+	}
+	DECLARE_DERIVED_EVENT( FCollectionManager, ICollectionManager::FAssetsAddedEvent, FAssetsAddedEvent ); 
+	virtual FAssetsAddedEvent& OnAssetsAdded() override { return AssetsAddedEvent; }
+
+	DECLARE_DERIVED_EVENT( FCollectionManager, ICollectionManager::FAssetsRemovedEvent, FAssetsRemovedEvent );
+	virtual FAssetsRemovedEvent& OnAssetsRemoved() override { return AssetsRemovedEvent; }
+PRAGMA_ENABLE_DEPRECATION_WARNINGS
+
 	virtual bool HasCollections() const override;
 	virtual void GetCollections(TArray<FCollectionNameType>& OutCollections) const override;
 	virtual void GetCollections(FName CollectionName, TArray<FCollectionNameType>& OutCollections) const override;
@@ -125,23 +222,23 @@ public:
 	virtual void GetChildCollectionNames(FName CollectionName, ECollectionShareType::Type ShareType, ECollectionShareType::Type ChildShareType, TArray<FName>& CollectionNames) const override;
 	virtual TOptional<FCollectionNameType> GetParentCollection(FName CollectionName, ECollectionShareType::Type ShareType) const override;
 	virtual bool CollectionExists(FName CollectionName, ECollectionShareType::Type ShareType) const override;
-	virtual bool GetAssetsInCollection(FName CollectionName, ECollectionShareType::Type ShareType, TArray<FName>& AssetPaths, ECollectionRecursionFlags::Flags RecursionMode = ECollectionRecursionFlags::Self) const override;
-	virtual bool GetObjectsInCollection(FName CollectionName, ECollectionShareType::Type ShareType, TArray<FName>& ObjectPaths, ECollectionRecursionFlags::Flags RecursionMode = ECollectionRecursionFlags::Self) const override;
-	virtual bool GetClassesInCollection(FName CollectionName, ECollectionShareType::Type ShareType, TArray<FName>& ClassPaths, ECollectionRecursionFlags::Flags RecursionMode = ECollectionRecursionFlags::Self) const override;
-	virtual void GetCollectionsContainingObject(FName ObjectPath, ECollectionShareType::Type ShareType, TArray<FName>& OutCollectionNames, ECollectionRecursionFlags::Flags RecursionMode = ECollectionRecursionFlags::Self) const override;
-	virtual void GetCollectionsContainingObject(FName ObjectPath, TArray<FCollectionNameType>& OutCollections, ECollectionRecursionFlags::Flags RecursionMode = ECollectionRecursionFlags::Self) const override;
-	virtual void GetCollectionsContainingObjects(const TArray<FName>& ObjectPaths, TMap<FCollectionNameType, TArray<FName>>& OutCollectionsAndMatchedObjects, ECollectionRecursionFlags::Flags RecursionMode = ECollectionRecursionFlags::Self) const override;
-	virtual FString GetCollectionsStringForObject(FName ObjectPath, ECollectionShareType::Type ShareType, ECollectionRecursionFlags::Flags RecursionMode = ECollectionRecursionFlags::Self, bool bFullPaths = true) const override;
+	virtual bool GetAssetsInCollection(FName CollectionName, ECollectionShareType::Type ShareType, TArray<FSoftObjectPath>& AssetPaths, ECollectionRecursionFlags::Flags RecursionMode = ECollectionRecursionFlags::Self) const override;
+	virtual bool GetObjectsInCollection(FName CollectionName, ECollectionShareType::Type ShareType, TArray<FSoftObjectPath>& ObjectPaths, ECollectionRecursionFlags::Flags RecursionMode = ECollectionRecursionFlags::Self) const override;
+	virtual bool GetClassesInCollection(FName CollectionName, ECollectionShareType::Type ShareType, TArray<FTopLevelAssetPath>& ClassPaths, ECollectionRecursionFlags::Flags RecursionMode = ECollectionRecursionFlags::Self) const override;
+	virtual void GetCollectionsContainingObject(const FSoftObjectPath& ObjectPath, ECollectionShareType::Type ShareType, TArray<FName>& OutCollectionNames, ECollectionRecursionFlags::Flags RecursionMode = ECollectionRecursionFlags::Self) const override;
+	virtual void GetCollectionsContainingObject(const FSoftObjectPath& ObjectPath, TArray<FCollectionNameType>& OutCollections, ECollectionRecursionFlags::Flags RecursionMode = ECollectionRecursionFlags::Self) const override;
+	virtual void GetCollectionsContainingObjects(const TArray<FSoftObjectPath>& ObjectPaths, TMap<FCollectionNameType, TArray<FSoftObjectPath>>& OutCollectionsAndMatchedObjects, ECollectionRecursionFlags::Flags RecursionMode = ECollectionRecursionFlags::Self) const override;
+	virtual FString GetCollectionsStringForObject(const FSoftObjectPath& ObjectPath, ECollectionShareType::Type ShareType, ECollectionRecursionFlags::Flags RecursionMode = ECollectionRecursionFlags::Self, bool bFullPaths = true) const override;
 	virtual void CreateUniqueCollectionName(const FName& BaseName, ECollectionShareType::Type ShareType, FName& OutCollectionName) const override;
 	virtual bool IsValidCollectionName(const FString& CollectionName, ECollectionShareType::Type ShareType) const override;
 	virtual bool CreateCollection(FName CollectionName, ECollectionShareType::Type ShareType, ECollectionStorageMode::Type StorageMode) override;
 	virtual bool RenameCollection(FName CurrentCollectionName, ECollectionShareType::Type CurrentShareType, FName NewCollectionName, ECollectionShareType::Type NewShareType) override;
 	virtual bool ReparentCollection(FName CollectionName, ECollectionShareType::Type ShareType, FName ParentCollectionName, ECollectionShareType::Type ParentShareType) override;
 	virtual bool DestroyCollection(FName CollectionName, ECollectionShareType::Type ShareType) override;
-	virtual bool AddToCollection(FName CollectionName, ECollectionShareType::Type ShareType, FName ObjectPath) override;
-	virtual bool AddToCollection(FName CollectionName, ECollectionShareType::Type ShareType, const TArray<FName>& ObjectPaths, int32* OutNumAdded = nullptr) override;
-	virtual bool RemoveFromCollection(FName CollectionName, ECollectionShareType::Type ShareType, FName ObjectPath) override;
-	virtual bool RemoveFromCollection(FName CollectionName, ECollectionShareType::Type ShareType, const TArray<FName>& ObjectPaths, int32* OutNumRemoved = nullptr) override;
+	virtual bool AddToCollection(FName CollectionName, ECollectionShareType::Type ShareType, const FSoftObjectPath& ObjectPath) override;
+	virtual bool AddToCollection(FName CollectionName, ECollectionShareType::Type ShareType, TConstArrayView<FSoftObjectPath> ObjectPaths, int32* OutNumAdded = nullptr) override;
+	virtual bool RemoveFromCollection(FName CollectionName, ECollectionShareType::Type ShareType, const FSoftObjectPath& ObjectPath) override;
+	virtual bool RemoveFromCollection(FName CollectionName, ECollectionShareType::Type ShareType, TConstArrayView<FSoftObjectPath> ObjectPaths, int32* OutNumRemoved = nullptr) override;
 	virtual bool SetDynamicQueryText(FName CollectionName, ECollectionShareType::Type ShareType, const FString& InQueryText) override;
 	virtual bool GetDynamicQueryText(FName CollectionName, ECollectionShareType::Type ShareType, FString& OutQueryText) const override;
 	virtual bool TestDynamicQuery(FName CollectionName, ECollectionShareType::Type ShareType, const ITextFilterExpressionContext& InContext, bool& OutResult) const override;
@@ -153,13 +250,13 @@ public:
 	virtual bool GetCollectionColor(FName CollectionName, ECollectionShareType::Type ShareType, TOptional<FLinearColor>& OutColor) const override;
 	virtual bool SetCollectionColor(FName CollectionName, ECollectionShareType::Type ShareType, const TOptional<FLinearColor>& NewColor) override;
 	virtual bool GetCollectionStorageMode(FName CollectionName, ECollectionShareType::Type ShareType, ECollectionStorageMode::Type& OutStorageMode) const override;
-	virtual bool IsObjectInCollection(FName ObjectPath, FName CollectionName, ECollectionShareType::Type ShareType, ECollectionRecursionFlags::Flags RecursionMode = ECollectionRecursionFlags::Self) const override;
+	virtual bool IsObjectInCollection(const FSoftObjectPath& ObjectPath, FName CollectionName, ECollectionShareType::Type ShareType, ECollectionRecursionFlags::Flags RecursionMode = ECollectionRecursionFlags::Self) const override;
 	virtual bool IsValidParentCollection(FName CollectionName, ECollectionShareType::Type ShareType, FName ParentCollectionName, ECollectionShareType::Type ParentShareType) const override;
 	virtual FText GetLastError() const override { return LastError; }
 	virtual void HandleFixupRedirectors(ICollectionRedirectorFollower& InRedirectorFollower) override;
-	virtual bool HandleRedirectorDeleted(const FName& ObjectPath) override;
-	virtual void HandleObjectRenamed(const FName& OldObjectPath, const FName& NewObjectPath) override;
-	virtual void HandleObjectDeleted(const FName& ObjectPath) override;
+	virtual bool HandleRedirectorDeleted(const FSoftObjectPath& ObjectPath) override;
+	virtual void HandleObjectRenamed(const FSoftObjectPath& OldObjectPath, const FSoftObjectPath& NewObjectPath) override;
+	virtual void HandleObjectDeleted(const FSoftObjectPath& ObjectPath) override;
 
 	/** Event for when collections are created */
 	DECLARE_DERIVED_EVENT( FCollectionManager, ICollectionManager::FCollectionCreatedEvent, FCollectionCreatedEvent );
@@ -170,12 +267,10 @@ public:
 	virtual FCollectionDestroyedEvent& OnCollectionDestroyed() override { return CollectionDestroyedEvent; }
 
 	/** Event for when assets are added to a collection */
-	DECLARE_DERIVED_EVENT( FCollectionManager, ICollectionManager::FAssetsAddedEvent, FAssetsAddedEvent ); 
-	virtual FAssetsAddedEvent& OnAssetsAdded() override { return AssetsAddedEvent; }
+	virtual FOnAssetsAddedToCollection& OnAssetsAddedToCollection() override { return AssetsAddedToCollectionDelegate; }
 
-	/** Event for when assets are removed to a collection */
-	DECLARE_DERIVED_EVENT( FCollectionManager, ICollectionManager::FAssetsRemovedEvent, FAssetsRemovedEvent );
-	virtual FAssetsRemovedEvent& OnAssetsRemoved() override { return AssetsRemovedEvent; }
+	/** Event for when assets are removed from a collection */
+	virtual FOnAssetsRemovedFromCollection& OnAssetsRemovedFromCollection() override { return AssetsRemovedFromCollectionDelegate; }
 
 	/** Event for when collections are renamed */
 	DECLARE_DERIVED_EVENT( FCollectionManager, ICollectionManager::FCollectionRenamedEvent, FCollectionRenamedEvent );
@@ -213,10 +308,10 @@ private:
 	bool RemoveCollection(const TSharedRef<FCollection>& CollectionRef, ECollectionShareType::Type ShareType);
 
 	/** Removes an object from any collections that contain it */
-	void RemoveObjectFromCollections(const FName& ObjectPath, TArray<FCollectionNameType>& OutUpdatedCollections);
+	void RemoveObjectFromCollections(const FSoftObjectPath& ObjectPath, TArray<FCollectionNameType>& OutUpdatedCollections);
 
 	/** Replaces an object with another in any collections that contain it */
-	void ReplaceObjectInCollections(const FName& OldObjectPath, const FName& NewObjectPath, TArray<FCollectionNameType>& OutUpdatedCollections);
+	void ReplaceObjectInCollections(const FSoftObjectPath& OldObjectPath, const FSoftObjectPath& NewObjectPath, TArray<FCollectionNameType>& OutUpdatedCollections);
 
 	/** Internal common functionality for saving a collection */
 	bool InternalSaveCollection(const TSharedRef<FCollection>& CollectionRef, FText& OutError);
@@ -244,9 +339,15 @@ private:
 	mutable FText LastError;
 
 	/** Event for when assets are added to a collection */
+	FOnAssetsAddedToCollection AssetsAddedToCollectionDelegate;
+
+	UE_DEPRECATED(5.1, "This event has been replaced by AssetsAddedToCollectionEvent")
 	FAssetsAddedEvent AssetsAddedEvent;
 
 	/** Event for when assets are removed from a collection */
+	FOnAssetsRemovedFromCollection AssetsRemovedFromCollectionDelegate;
+
+	UE_DEPRECATED(5.1, "This event has been replaced by AssetsRemovedFromCollectionEvent")
 	FAssetsRemovedEvent AssetsRemovedEvent;
 
 	/** Event for when collections are renamed */
diff --git a/Engine/Source/Developer/CollectionManager/Private/CollectionManagerConsoleCommands.h b/Engine/Source/Developer/CollectionManager/Private/CollectionManagerConsoleCommands.h
--- a/Engine/Source/Developer/CollectionManager/Private/CollectionManagerConsoleCommands.h
+++ b/Engine/Source/Developer/CollectionManager/Private/CollectionManagerConsoleCommands.h
@@ -138,7 +138,7 @@ public:
 
 		FName CollectionName = FName(*Args[0]);
 		FString ShareStr = Args[1];
-		FName ObjectPath = FName(*Args[2]);
+		FSoftObjectPath ObjectPath(Args[2]);
 		ECollectionShareType::Type Type;
 
 		if ( ShareStr == TEXT("LOCAL") )
@@ -179,7 +179,7 @@ public:
 
 		FName CollectionName = FName(*Args[0]);
 		FString ShareStr = Args[1];
-		FName ObjectPath = FName(*Args[2]);
+		FSoftObjectPath ObjectPath(Args[2]);
 
 		ECollectionShareType::Type Type;
 
diff --git a/Engine/Source/Developer/CollectionManager/Public/CollectionManagerTypes.h b/Engine/Source/Developer/CollectionManager/Public/CollectionManagerTypes.h
--- a/Engine/Source/Developer/CollectionManager/Public/CollectionManagerTypes.h
+++ b/Engine/Source/Developer/CollectionManager/Public/CollectionManagerTypes.h
@@ -3,6 +3,7 @@
 #pragma once
 
 #include "CoreMinimal.h"
+#include "UObject/SoftObjectPath.h"
 #include "ISourceControlState.h"
 
 struct ECollectionStorageMode
@@ -260,5 +261,5 @@ public:
 	virtual ~ICollectionRedirectorFollower() {}
 
 	/** Given an object path, will see if it needs to follow any redirectors, and if so, will populate OutNewObjectPath with the new name and return true */
-	virtual bool FixupObject(const FName& InObjectPath, FName& OutNewObjectPath) = 0;
+	virtual bool FixupObject(const FSoftObjectPath& InObjectPath, FSoftObjectPath& OutNewObjectPath) = 0;
 };
diff --git a/Engine/Source/Developer/CollectionManager/Public/ICollectionManager.h b/Engine/Source/Developer/CollectionManager/Public/ICollectionManager.h
--- a/Engine/Source/Developer/CollectionManager/Public/ICollectionManager.h
+++ b/Engine/Source/Developer/CollectionManager/Public/ICollectionManager.h
@@ -44,24 +44,38 @@ public:
 	virtual bool CollectionExists(FName CollectionName, ECollectionShareType::Type ShareType) const = 0;
 
 	/** Returns a list of asset paths found in the specified collection and share type */
+	virtual bool GetAssetsInCollection(FName CollectionName, ECollectionShareType::Type ShareType, TArray<FSoftObjectPath>& OutAssetPaths, ECollectionRecursionFlags::Flags RecursionMode = ECollectionRecursionFlags::Self) const = 0;
+	UE_DEPRECATED(5.1, "FNames containing full object paths are deprecated. Use FSoftObjectPath instead.")
 	virtual bool GetAssetsInCollection(FName CollectionName, ECollectionShareType::Type ShareType, TArray<FName>& OutAssetPaths, ECollectionRecursionFlags::Flags RecursionMode = ECollectionRecursionFlags::Self) const = 0;
 
 	/** Returns a list of class paths found in the specified collection and share type */
+	virtual bool GetClassesInCollection(FName CollectionName, ECollectionShareType::Type ShareType, TArray<FTopLevelAssetPath>& ClassPaths, ECollectionRecursionFlags::Flags RecursionMode = ECollectionRecursionFlags::Self) const = 0;
+	UE_DEPRECATED(5.1, "FNames containing class paths are deprecated. Use FTopLevelAssetPath instead.")
 	virtual bool GetClassesInCollection(FName CollectionName, ECollectionShareType::Type ShareType, TArray<FName>& ClassPaths, ECollectionRecursionFlags::Flags RecursionMode = ECollectionRecursionFlags::Self) const = 0;
 
 	/** Returns a list of object paths found in the specified collection and share type */
+	virtual bool GetObjectsInCollection(FName CollectionName, ECollectionShareType::Type ShareType, TArray<FSoftObjectPath>& ObjectPaths, ECollectionRecursionFlags::Flags RecursionMode = ECollectionRecursionFlags::Self) const = 0;
+	UE_DEPRECATED(5.1, "FNames containing full object paths are deprecated. Use FSoftObjectPath instead.")
 	virtual bool GetObjectsInCollection(FName CollectionName, ECollectionShareType::Type ShareType, TArray<FName>& ObjectPaths, ECollectionRecursionFlags::Flags RecursionMode = ECollectionRecursionFlags::Self) const = 0;
 
 	/** Returns a list of collections in which the specified object exists of the specified share type */
+	virtual void GetCollectionsContainingObject(const FSoftObjectPath& ObjectPath, ECollectionShareType::Type ShareType, TArray<FName>& OutCollectionNames, ECollectionRecursionFlags::Flags RecursionMode = ECollectionRecursionFlags::Self) const = 0;
+	UE_DEPRECATED(5.1, "FNames containing full object paths are deprecated. Use FSoftObjectPath instead.")
 	virtual void GetCollectionsContainingObject(FName ObjectPath, ECollectionShareType::Type ShareType, TArray<FName>& OutCollectionNames, ECollectionRecursionFlags::Flags RecursionMode = ECollectionRecursionFlags::Self) const = 0;
 
 	/** Returns a list of collections in which the specified object exists */
+	virtual void GetCollectionsContainingObject(const FSoftObjectPath& ObjectPath, TArray<FCollectionNameType>& OutCollections, ECollectionRecursionFlags::Flags RecursionMode = ECollectionRecursionFlags::Self) const = 0;
+	UE_DEPRECATED(5.1, "FNames containing full object paths are deprecated. Use FSoftObjectPath instead.")
 	virtual void GetCollectionsContainingObject(FName ObjectPath, TArray<FCollectionNameType>& OutCollections, ECollectionRecursionFlags::Flags RecursionMode = ECollectionRecursionFlags::Self) const = 0;
 
 	/** Returns a list of collections in which any of the specified objects exist */
+	virtual void GetCollectionsContainingObjects(const TArray<FSoftObjectPath>& ObjectPaths, TMap<FCollectionNameType, TArray<FSoftObjectPath>>& OutCollectionsAndMatchedObjects, ECollectionRecursionFlags::Flags RecursionMode = ECollectionRecursionFlags::Self) const = 0;
+	UE_DEPRECATED(5.1, "FNames containing full object paths are deprecated. Use FSoftObjectPath instead.")
 	virtual void GetCollectionsContainingObjects(const TArray<FName>& ObjectPaths, TMap<FCollectionNameType, TArray<FName>>& OutCollectionsAndMatchedObjects, ECollectionRecursionFlags::Flags RecursionMode = ECollectionRecursionFlags::Self) const = 0;
 
 	/** Returns a string containing a comma separated list of collections in which the specified object exists of the specified share type */
+	virtual FString GetCollectionsStringForObject(const FSoftObjectPath& ObjectPath, ECollectionShareType::Type ShareType, ECollectionRecursionFlags::Flags RecursionMode = ECollectionRecursionFlags::Self, bool bFullPaths = true) const = 0;
+	UE_DEPRECATED(5.1, "FNames containing full object paths are deprecated. Use FSoftObjectPath instead.")
 	virtual FString GetCollectionsStringForObject(FName ObjectPath, ECollectionShareType::Type ShareType, ECollectionRecursionFlags::Flags RecursionMode = ECollectionRecursionFlags::Self, bool bFullPaths = true) const = 0;
 
 	/** Creates a unique collection name for the given type taking the form BaseName+(unique number) */
@@ -120,7 +134,11 @@ public:
 	  * @param OutNumAdded if non-NULL, the number of objects successfully added to the collection
 	  * @return true if the add was successful. If false, GetLastError will return a human readable string description of the error.
 	  */
+	virtual bool AddToCollection(FName CollectionName, ECollectionShareType::Type ShareType, const FSoftObjectPath& ObjectPath) = 0;
+	virtual bool AddToCollection(FName CollectionName, ECollectionShareType::Type ShareType, TConstArrayView<FSoftObjectPath> ObjectPaths, int32* OutNumAdded = nullptr) = 0;
+	UE_DEPRECATED(5.1, "FNames containing full object paths are deprecated. Use FSoftObjectPath instead.")
 	virtual bool AddToCollection(FName CollectionName, ECollectionShareType::Type ShareType, FName ObjectPath) = 0;
+	UE_DEPRECATED(5.1, "FNames containing full object paths are deprecated. Use FSoftObjectPath instead.")
 	virtual bool AddToCollection(FName CollectionName, ECollectionShareType::Type ShareType, const TArray<FName>& ObjectPaths, int32* OutNumAdded = nullptr) = 0;
 
 	/**
@@ -132,7 +150,11 @@ public:
 	  * @param OutNumRemoved if non-NULL, the number of objects successfully removed from the collection
 	  * @return true if the remove was successful. If false, GetLastError will return a human readable string description of the error.
 	  */
+	virtual bool RemoveFromCollection(FName CollectionName, ECollectionShareType::Type ShareType, const FSoftObjectPath& ObjectPath) = 0;
+	virtual bool RemoveFromCollection(FName CollectionName, ECollectionShareType::Type ShareType, TConstArrayView<FSoftObjectPath> ObjectPaths, int32* OutNumRemoved = nullptr) = 0;
+	UE_DEPRECATED(5.1, "FNames containing full object paths are deprecated. Use FSoftObjectPath instead.")
 	virtual bool RemoveFromCollection(FName CollectionName, ECollectionShareType::Type ShareType, FName ObjectPath) = 0;
+	UE_DEPRECATED(5.1, "FNames containing full object paths are deprecated. Use FSoftObjectPath instead.")
 	virtual bool RemoveFromCollection(FName CollectionName, ECollectionShareType::Type ShareType, const TArray<FName>& ObjectPaths, int32* OutNumRemoved = nullptr) = 0;
 
 	/**
@@ -255,6 +277,8 @@ public:
 	 * @param ShareType The way the collection is shared.
 	 * @return true if the object is in the collection.
 	 */
+	virtual bool IsObjectInCollection(const FSoftObjectPath& ObjectPath, FName CollectionName, ECollectionShareType::Type ShareType, ECollectionRecursionFlags::Flags RecursionMode = ECollectionRecursionFlags::Self) const = 0;
+	UE_DEPRECATED(5.1, "FNames containing full object paths are deprecated. Use FSoftObjectPath instead.")
 	virtual bool IsObjectInCollection(FName ObjectPath, FName CollectionName, ECollectionShareType::Type ShareType, ECollectionRecursionFlags::Flags RecursionMode = ECollectionRecursionFlags::Self) const = 0;
 
 	/**
@@ -281,12 +305,18 @@ public:
 	 * Called to notify the collections that a redirector has been deleted and that they should ensure their on-disk representation is re-saved with the fixed up in-memory version
 	 * @return true if all of the collections that were referencing this redirector could be re-saved, false otherwise
 	 */
+	virtual bool HandleRedirectorDeleted(const FSoftObjectPath& ObjectPath) = 0;
+	UE_DEPRECATED(5.1, "FNames containing full object paths are deprecated. Use FSoftObjectPath instead.")
 	virtual bool HandleRedirectorDeleted(const FName& ObjectPath) = 0;
 
 	/** Called to notify the collections that an object has been renamed or moved */
+	virtual void HandleObjectRenamed(const FSoftObjectPath& OldObjectPath, const FSoftObjectPath& NewObjectPath) = 0;
+	UE_DEPRECATED(5.1, "FNames containing full object paths are deprecated. Use FSoftObjectPath instead.")
 	virtual void HandleObjectRenamed(const FName& OldObjectPath, const FName& NewObjectPath) = 0;
 
 	/** Called to notify the collections that an object has been deleted */
+	virtual void HandleObjectDeleted(const FSoftObjectPath& ObjectPath) = 0;
+	UE_DEPRECATED(5.1, "FNames containing full object paths are deprecated. Use FSoftObjectPath instead.")
 	virtual void HandleObjectDeleted(const FName& ObjectPath) = 0;
 
 	/** Event for when collections are created */
@@ -298,11 +328,19 @@ public:
 	virtual FCollectionDestroyedEvent& OnCollectionDestroyed() = 0;
 
 	/** Event for when assets are added to a collection */
+	DECLARE_MULTICAST_DELEGATE_TwoParams(FOnAssetsAddedToCollection, const FCollectionNameType&, TConstArrayView<FSoftObjectPath> );
+	virtual FOnAssetsAddedToCollection& OnAssetsAddedToCollection() = 0;
+
 	DECLARE_EVENT_TwoParams( ICollectionManager, FAssetsAddedEvent, const FCollectionNameType&, const TArray<FName>& );
+	UE_DEPRECATED(5.1, "Names containing full asset paths are deprecated, use OnAssetsAddedToCollection instead.")
 	virtual FAssetsAddedEvent& OnAssetsAdded() = 0;
 
 	/** Event for when assets are removed from a collection */
+	DECLARE_MULTICAST_DELEGATE_TwoParams( FOnAssetsRemovedFromCollection, const FCollectionNameType&, TConstArrayView<FSoftObjectPath> );
+	virtual FOnAssetsRemovedFromCollection& OnAssetsRemovedFromCollection() = 0;
+
 	DECLARE_EVENT_TwoParams( ICollectionManager, FAssetsRemovedEvent, const FCollectionNameType&, const TArray<FName>& );
+	UE_DEPRECATED(5.1, "Names containing full asset paths are deprecated, use OnAssetsRemovedFromCollection instead.")
 	virtual FAssetsRemovedEvent& OnAssetsRemoved() = 0;
 
 	/** Event for when collections are renamed */
diff --git a/Engine/Source/Editor/ContentBrowser/Private/AssetContextMenu.cpp b/Engine/Source/Editor/ContentBrowser/Private/AssetContextMenu.cpp
--- a/Engine/Source/Editor/ContentBrowser/Private/AssetContextMenu.cpp
+++ b/Engine/Source/Editor/ContentBrowser/Private/AssetContextMenu.cpp
@@ -750,18 +750,16 @@ bool FAssetContextMenu::AddCollectionMenuOptions(UToolMenu* Menu)
 	{
 		TSharedRef<FCollectionAssetManagement> QuickAssetManagement = MakeShared<FCollectionAssetManagement>();
 
-		TArray<FName> SelectedItemCollectionIds;
+		TArray<FSoftObjectPath> SelectedItemCollectionIds;
 		for (const FContentBrowserItem& SelectedItem : SelectedFiles)
 		{
-			FName ItemCollectionId;
+			FSoftObjectPath ItemCollectionId;
 			if (SelectedItem.TryGetCollectionId(ItemCollectionId))
 			{
 				SelectedItemCollectionIds.Add(ItemCollectionId);
 			}
 		}
-PRAGMA_DISABLE_DEPRECATION_WARNINGS
-		QuickAssetManagement->SetCurrentAssetPaths(UE::SoftObjectPath::Private::ConvertObjectPathNames(SelectedItemCollectionIds));
-PRAGMA_ENABLE_DEPRECATION_WARNINGS
+		QuickAssetManagement->SetCurrentAssetPaths(SelectedItemCollectionIds);
 
 		Section.AddSubMenu(
 			"ManageCollections",
@@ -1111,10 +1109,10 @@ void FAssetContextMenu::ExecuteRemoveFromCollection()
 {
 	if ( ensure(SourcesData.Collections.Num() == 1) )
 	{
-		TArray<FName> SelectedItemCollectionIds;
+		TArray<FSoftObjectPath> SelectedItemCollectionIds;
 		for (const FContentBrowserItem& SelectedItem : SelectedFiles)
 		{
-			FName ItemCollectionId;
+			FSoftObjectPath ItemCollectionId;
 			if (SelectedItem.TryGetCollectionId(ItemCollectionId))
 			{
 				SelectedItemCollectionIds.Add(ItemCollectionId);
diff --git a/Engine/Source/Editor/ContentBrowser/Private/CollectionAssetRegistryBridge.cpp b/Engine/Source/Editor/ContentBrowser/Private/CollectionAssetRegistryBridge.cpp
--- a/Engine/Source/Editor/ContentBrowser/Private/CollectionAssetRegistryBridge.cpp
+++ b/Engine/Source/Editor/ContentBrowser/Private/CollectionAssetRegistryBridge.cpp
@@ -36,9 +36,9 @@ public:
 	{
 	}
 
-	virtual bool FixupObject(const FName& InObjectPath, FName& OutNewObjectPath) override
+	virtual bool FixupObject(const FSoftObjectPath& InObjectPath, FSoftObjectPath& OutNewObjectPath) override
 	{
-		OutNewObjectPath = NAME_None;
+		OutNewObjectPath = FSoftObjectPath();
 
 		FString InObjectPathString = InObjectPath.ToString();
 		if (!InObjectPathString.StartsWith(TEXT("/")))
@@ -52,6 +52,7 @@ public:
 					const FString NewClassName = FLinkerLoad::FindNewPathNameForClass(ClassPathStr, false);
 					if (!NewClassName.IsEmpty())
 					{
+						check(FPackageName::IsValidObjectPath(NewClassName));
 						// Our new class name might be lacking the path, so try and find it so we can use the full path in the collection
 						UClass* FoundClass = FindObject<UClass>(nullptr, *NewClassName);
 						if (FoundClass)
@@ -133,7 +134,7 @@ PRAGMA_DISABLE_DEPRECATION_WARNINGS
 PRAGMA_ENABLE_DEPRECATION_WARNINGS
 		}
 
-		return OutNewObjectPath != NAME_None && InObjectPath != OutNewObjectPath;
+		return OutNewObjectPath.IsValid() && InObjectPath != OutNewObjectPath;
 	}
 
 private:
@@ -183,9 +184,7 @@ void FCollectionAssetRegistryBridge::OnAssetRenamed(const FAssetData& AssetData,
 	FCollectionManagerModule& CollectionManagerModule = FCollectionManagerModule::GetModule();
 
 	// Notify the collections manager that an asset has been renamed
-PRAGMA_DISABLE_DEPRECATION_WARNINGS
-	CollectionManagerModule.Get().HandleObjectRenamed(*OldObjectPath, AssetData.ObjectPath);
-PRAGMA_ENABLE_DEPRECATION_WARNINGS
+	CollectionManagerModule.Get().HandleObjectRenamed(FSoftObjectPath(OldObjectPath), AssetData.GetSoftObjectPath());
 }
 
 void FCollectionAssetRegistryBridge::OnAssetRemoved(const FAssetData& AssetData)
@@ -196,16 +195,12 @@ void FCollectionAssetRegistryBridge::OnAssetRemoved(const FAssetData& AssetData)
 	{
 		// Notify the collections manager that a redirector has been removed
 		// This will attempt to re-save any collections that still have a reference to this redirector in their on-disk collection data
-PRAGMA_DISABLE_DEPRECATION_WARNINGS
-		CollectionManagerModule.Get().HandleRedirectorDeleted(AssetData.ObjectPath);
-PRAGMA_ENABLE_DEPRECATION_WARNINGS
+		CollectionManagerModule.Get().HandleRedirectorDeleted(AssetData.GetSoftObjectPath());
 	}
 	else
 	{
 		// Notify the collections manager that an asset has been removed
-PRAGMA_DISABLE_DEPRECATION_WARNINGS
-		CollectionManagerModule.Get().HandleObjectDeleted(AssetData.ObjectPath);
-PRAGMA_ENABLE_DEPRECATION_WARNINGS
+		CollectionManagerModule.Get().HandleObjectDeleted(AssetData.GetSoftObjectPath());
 	}
 }
 
diff --git a/Engine/Source/Editor/ContentBrowser/Private/FrontendFilters.cpp b/Engine/Source/Editor/ContentBrowser/Private/FrontendFilters.cpp
--- a/Engine/Source/Editor/ContentBrowser/Private/FrontendFilters.cpp
+++ b/Engine/Source/Editor/ContentBrowser/Private/FrontendFilters.cpp
@@ -306,7 +306,7 @@ public:
 
 		if (CollectionManager)
 		{
-			FName ItemCollectionId;
+			FSoftObjectPath ItemCollectionId;
 			if (AssetPtr->TryGetCollectionId(ItemCollectionId))
 			{
 				CollectionManager->GetCollectionsContainingObject(ItemCollectionId, ECollectionShareType::CST_All, AssetCollectionNames, ECollectionRecursionFlags::SelfAndChildren);
diff --git a/Engine/Source/Editor/ContentBrowser/Private/SAssetView.cpp b/Engine/Source/Editor/ContentBrowser/Private/SAssetView.cpp
--- a/Engine/Source/Editor/ContentBrowser/Private/SAssetView.cpp
+++ b/Engine/Source/Editor/ContentBrowser/Private/SAssetView.cpp
@@ -208,8 +208,8 @@ void SAssetView::Construct( const FArguments& InArgs )
 	ContentBrowserData->OnItemDataDiscoveryComplete().AddSP(this, &SAssetView::HandleItemDataDiscoveryComplete);
 
 	FCollectionManagerModule& CollectionManagerModule = FCollectionManagerModule::GetModule();
-	CollectionManagerModule.Get().OnAssetsAdded().AddSP( this, &SAssetView::OnAssetsAddedToCollection );
-	CollectionManagerModule.Get().OnAssetsRemoved().AddSP( this, &SAssetView::OnAssetsRemovedFromCollection );
+	CollectionManagerModule.Get().OnAssetsAddedToCollection().AddSP( this, &SAssetView::OnAssetsAddedToCollection );
+	CollectionManagerModule.Get().OnAssetsRemovedFromCollection().AddSP( this, &SAssetView::OnAssetsRemovedFromCollection );
 	CollectionManagerModule.Get().OnCollectionRenamed().AddSP( this, &SAssetView::OnCollectionRenamed );
 	CollectionManagerModule.Get().OnCollectionUpdated().AddSP( this, &SAssetView::OnCollectionUpdated );
 
@@ -1413,13 +1413,13 @@ FReply SAssetView::OnDragOver( const FGeometry& MyGeometry, const FDragDropEvent
 	}
 	else if (HasSingleCollectionSource())
 	{
-		TArray<FName> NewCollectionItems;
+		TArray<FSoftObjectPath> NewCollectionItems;
 
 		if (TSharedPtr<FContentBrowserDataDragDropOp> ContentDragDropOp = DragDropEvent.GetOperationAs<FContentBrowserDataDragDropOp>())
 		{
 			for (const FContentBrowserItem& DraggedItem : ContentDragDropOp->GetDraggedFiles())
 			{
-				FName CollectionItemId;
+				FSoftObjectPath CollectionItemId;
 				if (DraggedItem.TryGetCollectionId(CollectionItemId))
 				{
 					NewCollectionItems.Add(CollectionItemId);
@@ -1429,20 +1429,20 @@ FReply SAssetView::OnDragOver( const FGeometry& MyGeometry, const FDragDropEvent
 		else
 		{
 			const TArray<FAssetData> AssetDatas = AssetUtil::ExtractAssetDataFromDrag(DragDropEvent);
-			Algo::Transform(AssetDatas, NewCollectionItems, &FAssetData::ObjectPath);
+			Algo::Transform(AssetDatas, NewCollectionItems, &FAssetData::GetSoftObjectPath);
 		}
 
 		if (NewCollectionItems.Num() > 0)
 		{
 			if (TSharedPtr<FAssetDragDropOp> AssetDragDropOp = DragDropEvent.GetOperationAs<FAssetDragDropOp>())
 			{
-				TArray< FName > ObjectPaths;
+				TArray<FSoftObjectPath> ObjectPaths;
 				FCollectionManagerModule& CollectionManagerModule = FCollectionManagerModule::GetModule();
 				const FCollectionNameType& Collection = SourcesData.Collections[0];
 				CollectionManagerModule.Get().GetObjectsInCollection(Collection.Name, Collection.Type, ObjectPaths);
 
 				bool IsValidDrop = false;
-				for (const FName& NewCollectionItem : NewCollectionItems)
+				for (const FSoftObjectPath& NewCollectionItem : NewCollectionItems)
 				{
 					if (!ObjectPaths.Contains(NewCollectionItem))
 					{
@@ -1493,13 +1493,13 @@ FReply SAssetView::OnDrop( const FGeometry& MyGeometry, const FDragDropEvent& Dr
 	}
 	else if (HasSingleCollectionSource())
 	{
-		TArray<FName> NewCollectionItems;
+		TArray<FSoftObjectPath> NewCollectionItems;
 
 		if (TSharedPtr<FContentBrowserDataDragDropOp> ContentDragDropOp = DragDropEvent.GetOperationAs<FContentBrowserDataDragDropOp>())
 		{
 			for (const FContentBrowserItem& DraggedItem : ContentDragDropOp->GetDraggedFiles())
 			{
-				FName CollectionItemId;
+				FSoftObjectPath CollectionItemId;
 				if (DraggedItem.TryGetCollectionId(CollectionItemId))
 				{
 					NewCollectionItems.Add(CollectionItemId);
@@ -1509,7 +1509,7 @@ FReply SAssetView::OnDrop( const FGeometry& MyGeometry, const FDragDropEvent& Dr
 		else
 		{
 			const TArray<FAssetData> AssetDatas = AssetUtil::ExtractAssetDataFromDrag(DragDropEvent);
-			Algo::Transform(AssetDatas, NewCollectionItems, &FAssetData::ObjectPath);
+			Algo::Transform(AssetDatas, NewCollectionItems, &FAssetData::GetSoftObjectPath);
 		}
 
 		if (NewCollectionItems.Num() > 0)
@@ -2264,7 +2264,7 @@ void SAssetView::SetMajorityAssetType(FName NewMajorityAssetType)
 	}
 }
 
-void SAssetView::OnAssetsAddedToCollection( const FCollectionNameType& Collection, const TArray< FName >& ObjectPaths )
+void SAssetView::OnAssetsAddedToCollection( const FCollectionNameType& Collection, TConstArrayView<FSoftObjectPath> ObjectPaths )
 {
 	if ( !SourcesData.Collections.Contains( Collection ) )
 	{
@@ -2274,7 +2274,7 @@ void SAssetView::OnAssetsAddedToCollection( const FCollectionNameType& Collectio
 	RequestSlowFullListRefresh();
 }
 
-void SAssetView::OnAssetsRemovedFromCollection( const FCollectionNameType& Collection, const TArray< FName >& ObjectPaths )
+void SAssetView::OnAssetsRemovedFromCollection( const FCollectionNameType& Collection, TConstArrayView<FSoftObjectPath> ObjectPaths )
 {
 	if ( !SourcesData.Collections.Contains( Collection ) )
 	{
diff --git a/Engine/Source/Editor/ContentBrowser/Private/SCollectionView.cpp b/Engine/Source/Editor/ContentBrowser/Private/SCollectionView.cpp
--- a/Engine/Source/Editor/ContentBrowser/Private/SCollectionView.cpp
+++ b/Engine/Source/Editor/ContentBrowser/Private/SCollectionView.cpp
@@ -134,8 +134,8 @@ void SCollectionView::Construct( const FArguments& InArgs )
 	CollectionManagerModule.Get().OnCollectionReparented().AddSP( this, &SCollectionView::HandleCollectionReparented );
 	CollectionManagerModule.Get().OnCollectionDestroyed().AddSP( this, &SCollectionView::HandleCollectionDestroyed );
 	CollectionManagerModule.Get().OnCollectionUpdated().AddSP( this, &SCollectionView::HandleCollectionUpdated );
-	CollectionManagerModule.Get().OnAssetsAdded().AddSP( this, &SCollectionView::HandleAssetsAddedToCollection );
-	CollectionManagerModule.Get().OnAssetsRemoved().AddSP( this, &SCollectionView::HandleAssetsRemovedFromCollection );
+	CollectionManagerModule.Get().OnAssetsAddedToCollection().AddSP( this, &SCollectionView::HandleAssetsAddedToCollection );
+	CollectionManagerModule.Get().OnAssetsRemovedFromCollection().AddSP( this, &SCollectionView::HandleAssetsRemovedFromCollection );
 
 	ISourceControlModule::Get().RegisterProviderChanged(FSourceControlProviderChanged::FDelegate::CreateSP(this, &SCollectionView::HandleSourceControlProviderChanged));
 	SourceControlStateChangedDelegateHandle = ISourceControlModule::Get().GetProvider().RegisterSourceControlStateChanged_Handle(FSourceControlStateChanged::FDelegate::CreateSP(this, &SCollectionView::HandleSourceControlStateChanged));
@@ -296,12 +296,12 @@ void SCollectionView::HandleCollectionUpdated( const FCollectionNameType& Collec
 	}
 }
 
-void SCollectionView::HandleAssetsAddedToCollection( const FCollectionNameType& Collection, const TArray<FName>& AssetsAdded )
+void SCollectionView::HandleAssetsAddedToCollection( const FCollectionNameType& Collection, TConstArrayView<FSoftObjectPath> AssetsAdded )
 {
 	HandleCollectionUpdated(Collection);
 }
 
-void SCollectionView::HandleAssetsRemovedFromCollection( const FCollectionNameType& Collection, const TArray<FName>& AssetsRemoved )
+void SCollectionView::HandleAssetsRemovedFromCollection( const FCollectionNameType& Collection, TConstArrayView<FSoftObjectPath> AssetsRemoved )
 {
 	HandleCollectionUpdated(Collection);
 }
diff --git a/Engine/Source/Editor/ContentBrowser/Private/SCollectionView.h b/Engine/Source/Editor/ContentBrowser/Private/SCollectionView.h
--- a/Engine/Source/Editor/ContentBrowser/Private/SCollectionView.h
+++ b/Engine/Source/Editor/ContentBrowser/Private/SCollectionView.h
@@ -254,10 +254,10 @@ private:
 	void HandleCollectionUpdated( const FCollectionNameType& Collection );
 
 	/** Handles assets being added to a collection */
-	void HandleAssetsAddedToCollection( const FCollectionNameType& Collection, const TArray<FName>& AssetsAdded );
+	void HandleAssetsAddedToCollection( const FCollectionNameType& Collection, TConstArrayView<FSoftObjectPath> AssetsAdded );
 
 	/** Handles assets being removed from a collection */
-	void HandleAssetsRemovedFromCollection( const FCollectionNameType& Collection, const TArray<FName>& AssetsRemoved );
+	void HandleAssetsRemovedFromCollection( const FCollectionNameType& Collection, TConstArrayView<FSoftObjectPath> AssetsRemoved );
 
 	/** Handles the source control provider changing */
 	void HandleSourceControlProviderChanged(class ISourceControlProvider& OldProvider, class ISourceControlProvider& NewProvider);
diff --git a/Engine/Source/Editor/ContentBrowser/Private/SContentBrowser.cpp b/Engine/Source/Editor/ContentBrowser/Private/SContentBrowser.cpp
--- a/Engine/Source/Editor/ContentBrowser/Private/SContentBrowser.cpp
+++ b/Engine/Source/Editor/ContentBrowser/Private/SContentBrowser.cpp
@@ -2895,19 +2895,17 @@ void SContentBrowser::OnItemSelectionChanged(const FContentBrowserItem& Selected
 		AssetContextMenu->SetSelectedItems(SelectedItems);
 
 		{
-			TArray<FName> SelectedCollectionItems;
+			TArray<FSoftObjectPath> SelectedCollectionItems;
 			for (const FContentBrowserItem& SelectedAssetItem : SelectedItems)
 			{
-				FName CollectionItemId;
+				FSoftObjectPath CollectionItemId;
 				if (SelectedAssetItem.TryGetCollectionId(CollectionItemId))
 				{
 					SelectedCollectionItems.Add(CollectionItemId);
 				}
 			}
 
-PRAGMA_DISABLE_DEPRECATION_WARNINGS
-			CollectionViewPtr->SetSelectedAssetPaths(UE::SoftObjectPath::Private::ConvertObjectPathNames(SelectedCollectionItems));
-PRAGMA_ENABLE_DEPRECATION_WARNINGS
+			CollectionViewPtr->SetSelectedAssetPaths(SelectedCollectionItems);
 		}
 
 		if (AssetSelectionChangedDelegate.IsBound())
diff --git a/Engine/Source/Editor/ContentBrowser/Public/SAssetView.h b/Engine/Source/Editor/ContentBrowser/Public/SAssetView.h
--- a/Engine/Source/Editor/ContentBrowser/Public/SAssetView.h
+++ b/Engine/Source/Editor/ContentBrowser/Public/SAssetView.h
@@ -465,10 +465,10 @@ private:
 	void SetMajorityAssetType(FName NewMajorityAssetType);
 
 	/** Handler for when an asset is added to a collection */
-	void OnAssetsAddedToCollection( const FCollectionNameType& Collection, const TArray< FName >& ObjectPaths );
+	void OnAssetsAddedToCollection( const FCollectionNameType& Collection, TConstArrayView<FSoftObjectPath> ObjectPaths );
 
 	/** Handler for when an asset is removed from a collection */
-	void OnAssetsRemovedFromCollection( const FCollectionNameType& Collection, const TArray< FName >& ObjectPaths );
+	void OnAssetsRemovedFromCollection( const FCollectionNameType& Collection, TConstArrayView<FSoftObjectPath> ObjectPaths );
 
 	/** Handler for when a collection is renamed */
 	void OnCollectionRenamed( const FCollectionNameType& OriginalCollection, const FCollectionNameType& NewCollection );
diff --git a/Engine/Source/Editor/ContentBrowserData/Private/ContentBrowserDataSource.cpp b/Engine/Source/Editor/ContentBrowserData/Private/ContentBrowserDataSource.cpp
--- a/Engine/Source/Editor/ContentBrowserData/Private/ContentBrowserDataSource.cpp
+++ b/Engine/Source/Editor/ContentBrowserData/Private/ContentBrowserDataSource.cpp
@@ -485,7 +485,7 @@ bool UContentBrowserDataSource::HandleDragDropOnItem(const FContentBrowserItemDa
 	return false;
 }
 
-bool UContentBrowserDataSource::TryGetCollectionId(const FContentBrowserItemData& InItem, FName& OutCollectionId)
+bool UContentBrowserDataSource::TryGetCollectionId(const FContentBrowserItemData& InItem, FSoftObjectPath& OutCollectionId)
 {
 	return false;
 }
diff --git a/Engine/Source/Editor/ContentBrowserData/Private/ContentBrowserItem.cpp b/Engine/Source/Editor/ContentBrowserData/Private/ContentBrowserItem.cpp
--- a/Engine/Source/Editor/ContentBrowserData/Private/ContentBrowserItem.cpp
+++ b/Engine/Source/Editor/ContentBrowserData/Private/ContentBrowserItem.cpp
@@ -548,9 +548,16 @@ bool FContentBrowserItem::UpdateThumbnail(FAssetThumbnail& InThumbnail) const
 	return FContentBrowserItemHelper::CallDataSourceImpl<UContentBrowserDataSource>(*this, UE_PROJECTION_MEMBER(UContentBrowserDataSource, UpdateThumbnail), InThumbnail);
 }
 
+bool FContentBrowserItem::TryGetCollectionId(FSoftObjectPath& OutCollectionId) const
+{
+	return FContentBrowserItemHelper::CallDataSourceImpl<UContentBrowserDataSource>(*this, UE_PROJECTION_MEMBER(UContentBrowserDataSource, TryGetCollectionId), OutCollectionId);
+}
+
 bool FContentBrowserItem::TryGetCollectionId(FName& OutCollectionId) const
 {
+PRAGMA_DISABLE_DEPRECATION_WARNINGS
 	return FContentBrowserItemHelper::CallDataSourceImpl<UContentBrowserDataSource>(*this, UE_PROJECTION_MEMBER(UContentBrowserDataSource, TryGetCollectionId), OutCollectionId);
+PRAGMA_ENABLE_DEPRECATION_WARNINGS
 }
 
 bool FContentBrowserItem::Legacy_TryGetPackagePath(FName& OutPackagePath) const
diff --git a/Engine/Source/Editor/ContentBrowserData/Public/ContentBrowserDataSource.h b/Engine/Source/Editor/ContentBrowserData/Public/ContentBrowserDataSource.h
--- a/Engine/Source/Editor/ContentBrowserData/Public/ContentBrowserDataSource.h
+++ b/Engine/Source/Editor/ContentBrowserData/Public/ContentBrowserDataSource.h
@@ -667,11 +667,25 @@ public:
 	 * Attempt to retrieve the identifier that should be used when storing a reference to the given item within a collection.
 	 *
 	 * @param InItem The item to query.
-	 * @param InOutStr The collection ID to fill.
+	 * @param OutCollectionId The collection ID to fill.
 	 *
 	 * @return True if the ID was retrieved, false otherwise.
 	 */
-	virtual bool TryGetCollectionId(const FContentBrowserItemData& InItem, FName& OutCollectionId);
+	virtual bool TryGetCollectionId(const FContentBrowserItemData& InItem, FSoftObjectPath& OutCollectionId);
+
+	UE_DEPRECATED(5.1, "FNames containing full object paths are deprecated. Use FSoftObjectPath instead.")
+	bool TryGetCollectionId(const FContentBrowserItemData& InItem, FName& OutCollectionId)
+	{
+PRAGMA_DISABLE_DEPRECATION_WARNINGS
+		FSoftObjectPath Temp;
+		if (TryGetCollectionId(InItem, Temp))
+		{
+			OutCollectionId = Temp.ToFName();
+			return true;
+		}
+		return false;
+PRAGMA_ENABLE_DEPRECATION_WARNINGS
+	}
 
 	/**
 	 * Attempt to retrieve the package path associated with the given item.
diff --git a/Engine/Source/Editor/ContentBrowserData/Public/ContentBrowserItem.h b/Engine/Source/Editor/ContentBrowserData/Public/ContentBrowserItem.h
--- a/Engine/Source/Editor/ContentBrowserData/Public/ContentBrowserItem.h
+++ b/Engine/Source/Editor/ContentBrowserData/Public/ContentBrowserItem.h
@@ -394,6 +394,9 @@ public:
 	 *
 	 * @return True if the ID was retrieved, false otherwise.
 	 */
+	bool TryGetCollectionId(FSoftObjectPath& OutCollectionId) const;
+
+	UE_DEPRECATED(5.1, "FNames containing full object paths are deprecated. Use FSoftObjectPath instead.")
 	bool TryGetCollectionId(FName & OutCollectionId) const;
 
 	/**
diff --git a/Engine/Source/Editor/UnrealEd/Private/Commandlets/AssetRegistryGenerator.cpp b/Engine/Source/Editor/UnrealEd/Private/Commandlets/AssetRegistryGenerator.cpp
--- a/Engine/Source/Editor/UnrealEd/Private/Commandlets/AssetRegistryGenerator.cpp
+++ b/Engine/Source/Editor/UnrealEd/Private/Commandlets/AssetRegistryGenerator.cpp
@@ -1150,7 +1150,7 @@ void FAssetRegistryGenerator::UpdateCollectionAssetData()
 	}
 
 	// Build the list of collections we should tag for each asset
-	TMap<FName, TArray<FName>> AssetPathNamesToCollectionTags;
+	TMap<FSoftObjectPath, TArray<FName>> AssetPathsToCollectionTags;
 	{
 		ICollectionManager& CollectionManager = FCollectionManagerModule::GetModule().Get();
 
@@ -1171,17 +1171,17 @@ void FAssetRegistryGenerator::UpdateCollectionAssetData()
 			});
 		}
 		
-		TArray<FName> TmpAssetPathNames;
+		TArray<FSoftObjectPath> TmpAssetPaths;
 		for (const FCollectionNameType& CollectionNameToTag : CollectionNamesToTag)
 		{
 			const FName CollectionTagName = *FString::Printf(TEXT("%s%s"), FAssetData::GetCollectionTagPrefix(), *CollectionNameToTag.Name.ToString());
 
-			TmpAssetPathNames.Reset();
-			CollectionManager.GetAssetsInCollection(CollectionNameToTag.Name, CollectionNameToTag.Type, TmpAssetPathNames);
+			TmpAssetPaths.Reset();
+			CollectionManager.GetAssetsInCollection(CollectionNameToTag.Name, CollectionNameToTag.Type, TmpAssetPaths);
 
-			for (const FName& AssetPathName : TmpAssetPathNames)
+			for (const FSoftObjectPath& AssetPath : TmpAssetPaths)
 			{
-				TArray<FName>& CollectionTagsForAsset = AssetPathNamesToCollectionTags.FindOrAdd(AssetPathName);
+				TArray<FName>& CollectionTagsForAsset = AssetPathsToCollectionTags.FindOrAdd(AssetPath);
 				CollectionTagsForAsset.AddUnique(CollectionTagName);
 			}
 		}
@@ -1191,12 +1191,12 @@ void FAssetRegistryGenerator::UpdateCollectionAssetData()
 	// Collection tags are queried only by the existence of the key, the value is never used. But Tag Values are not allowed
 	// to be empty. Set the value for each tag to an arbitrary field, something short to avoid wasting memory. We use 1 (aka "true") for now.
 	FStringView CollectionValue(TEXTVIEW("1"));
-	for (const TPair<FName, TArray<FName>>& AssetPathNameToCollectionTagsPair : AssetPathNamesToCollectionTags)
+	for (const TPair<FSoftObjectPath, TArray<FName>>& AssetPathToCollectionTagsPair : AssetPathsToCollectionTags)
 	{
-		const FName AssetPathName = AssetPathNameToCollectionTagsPair.Key;
-		const TArray<FName>& CollectionTagsForAsset = AssetPathNameToCollectionTagsPair.Value;
+		const FSoftObjectPath& AssetPath = AssetPathToCollectionTagsPair.Key;
+		const TArray<FName>& CollectionTagsForAsset = AssetPathToCollectionTagsPair.Value;
 
-		const FAssetData* AssetData = State.GetAssetByObjectPath(FSoftObjectPath(WriteToString<FName::StringBufferSize>(AssetPathName)));
+		const FAssetData* AssetData = State.GetAssetByObjectPath(AssetPath);
 		if (AssetData)
 		{
 			FAssetDataTagMap TagsAndValues = AssetData->TagsAndValues.CopyMap();
@@ -2412,10 +2412,10 @@ void FAssetRegistryGenerator::WriteCollection(FName CollectionName, const TArray
 {
 	if (CreateOrEmptyCollection(CollectionName))
 	{
-		TArray<FName> AssetNames = PackageNames;
+		TArray<FSoftObjectPath> AssetPaths;
 
 		// Convert package names to asset names
-		for (FName& Name : AssetNames)
+		for (const FName& Name : PackageNames)
 		{
 			FString PackageName = Name.ToString();
 			int32 LastPathDelimiter;
@@ -2423,12 +2423,12 @@ void FAssetRegistryGenerator::WriteCollection(FName CollectionName, const TArray
 			{
 				const FString AssetName = PackageName.Mid(LastPathDelimiter + 1);
 				PackageName = PackageName + FString(TEXT(".")) + AssetName;
-				Name = *PackageName;
 			}
+			AssetPaths.Add(FSoftObjectPath(PackageName));
 		}
 
 		ICollectionManager& CollectionManager = FCollectionManagerModule::GetModule().Get();
-		CollectionManager.AddToCollection(CollectionName, ECollectionShareType::CST_Local, AssetNames);
+		CollectionManager.AddToCollection(CollectionName, ECollectionShareType::CST_Local, AssetPaths);
 
 		UE_LOG(LogAssetRegistryGenerator, Log, TEXT("Updated collection %s"), *CollectionName.ToString());
 	}
diff --git a/Engine/Source/Editor/UnrealEd/Private/Commandlets/BaseIteratePackagesCommandlet.cpp b/Engine/Source/Editor/UnrealEd/Private/Commandlets/BaseIteratePackagesCommandlet.cpp
--- a/Engine/Source/Editor/UnrealEd/Private/Commandlets/BaseIteratePackagesCommandlet.cpp
+++ b/Engine/Source/Editor/UnrealEd/Private/Commandlets/BaseIteratePackagesCommandlet.cpp
@@ -984,7 +984,7 @@ int32 UBaseIteratePackagesCommandlet::Main( const FString& Params )
 	if (!FilterByCollection.IsEmpty())
 	{
 		ICollectionManager& CollectionManager = FCollectionManagerModule::GetModule().Get();
-		TArray<FName> CollectionAssets;
+		TArray<FSoftObjectPath> CollectionAssets;
 		if (!CollectionManager.GetAssetsInCollection(FName(*FilterByCollection), ECollectionShareType::CST_All, CollectionAssets))
 		{
 			UE_LOG(LogIteratePackagesCommandlet, Warning, TEXT("Could not get assets in collection '%s'. Skipping filter."), *FilterByCollection);
@@ -992,9 +992,9 @@ int32 UBaseIteratePackagesCommandlet::Main( const FString& Params )
 		else
 		{
 			//insert all of the collection names into the set for fast filter checks
-			for (const FName &AssetName : CollectionAssets)
+			for (const FSoftObjectPath& AssetPath : CollectionAssets)
 			{
-				CollectionFilter.Add(FName(*FPackageName::ObjectPathToPackageName(AssetName.ToString())));
+				CollectionFilter.Add(AssetPath.GetLongPackageFName());
 			}
 		}
 	}
diff --git a/Engine/Source/Editor/UnrealEd/Private/Commandlets/ContentCommandlets.cpp b/Engine/Source/Editor/UnrealEd/Private/Commandlets/ContentCommandlets.cpp
--- a/Engine/Source/Editor/UnrealEd/Private/Commandlets/ContentCommandlets.cpp
+++ b/Engine/Source/Editor/UnrealEd/Private/Commandlets/ContentCommandlets.cpp
@@ -1142,7 +1142,7 @@ int32 UResavePackagesCommandlet::Main( const FString& Params )
 	if (!FilterByCollection.IsEmpty())
 	{
 		ICollectionManager& CollectionManager = FCollectionManagerModule::GetModule().Get();
-		TArray<FName> CollectionAssets;
+		TArray<FSoftObjectPath> CollectionAssets;
 		if (!CollectionManager.GetAssetsInCollection(FName(*FilterByCollection), ECollectionShareType::CST_All, CollectionAssets))
 		{
 			UE_LOG(LogContentCommandlet, Warning, TEXT("Could not get assets in collection '%s'. Skipping filter."), *FilterByCollection);
@@ -1150,9 +1150,9 @@ int32 UResavePackagesCommandlet::Main( const FString& Params )
 		else
 		{
 			//insert all of the collection names into the set for fast filter checks
-			for (const FName &AssetName : CollectionAssets)
+			for (const FSoftObjectPath& AssetPath: CollectionAssets)
 			{
-				CollectionFilter.Add(FName(*FPackageName::ObjectPathToPackageName(AssetName.ToString())));
+				CollectionFilter.Add(AssetPath.GetLongPackageFName());
 			}
 		}
 	}
@@ -2604,7 +2604,7 @@ int32 UWrangleContentCommandlet::Main( const FString& Params )
 		if (CollectionsToFullyLoadSection)
 		{
 			ICollectionManager& CollectionManager = FCollectionManagerModule::GetModule().Get();
-			TArray<FName> CollectionAssets;
+			TArray<FSoftObjectPath> CollectionAssets;
 			
 			for (FConfigSectionMap::TConstIterator CollectionIt(*CollectionsToFullyLoadSection); CollectionIt; ++CollectionIt)
 			{
@@ -2617,9 +2617,9 @@ int32 UWrangleContentCommandlet::Main( const FString& Params )
 				else
 				{
 					//insert all of the collection names into the set for fast filter checks
-					for (const FName &AssetName : CollectionAssets)
+					for (const FSoftObjectPath& AssetPath : CollectionAssets)
 					{
-						PackagesToFullyLoad.Add(TEXT("Package"), FPackageName::ObjectPathToPackageName(AssetName.ToString()));
+						PackagesToFullyLoad.Add(TEXT("Package"), AssetPath.GetLongPackageName());
 					}
 				}
 			}
@@ -3654,7 +3654,7 @@ int32 UStaticMeshMinLodCommandlet::Main(const FString& Params)
 
 			if (bGenerateCollections)
 			{
-				CollectionManager.AddToCollection(Pair.Key, ECollectionShareType::CST_Local, FName(*Info.Paths[i]));
+				CollectionManager.AddToCollection(Pair.Key, ECollectionShareType::CST_Local, FSoftObjectPath(Info.Paths[i]));
 			}
 		}
 
diff --git a/Engine/Source/Editor/UnrealEd/Private/Commandlets/DerivedDataCacheCommandlet.cpp b/Engine/Source/Editor/UnrealEd/Private/Commandlets/DerivedDataCacheCommandlet.cpp
--- a/Engine/Source/Editor/UnrealEd/Private/Commandlets/DerivedDataCacheCommandlet.cpp
+++ b/Engine/Source/Editor/UnrealEd/Private/Commandlets/DerivedDataCacheCommandlet.cpp
@@ -488,12 +488,12 @@ int32 UDerivedDataCacheCommandlet::Main( const FString& Params )
 					continue;
 				}
 
-				TArray<FName> FoundAssets;
+				TArray<FSoftObjectPath> FoundAssets;
 				CollectionManager.GetAssetsInCollection(*CollectionName, ECollectionShareType::CST_All, FoundAssets, ECollectionRecursionFlags::SelfAndChildren);
 				Tokens.Reserve(Tokens.Num() + FoundAssets.Num());
-				for (FName AssetName : FoundAssets)
+				for (const FSoftObjectPath& AssetPath : FoundAssets)
 				{
-					CommandLinePackageNames.Add(FPackageName::ObjectPathToPackageName(AssetName.ToString()));
+					CommandLinePackageNames.Add(AssetPath.GetLongPackageName());
 				}
 			}
 		}
diff --git a/Engine/Source/Editor/UnrealEd/Private/Toolkits/AssetEditorToolkit.cpp b/Engine/Source/Editor/UnrealEd/Private/Toolkits/AssetEditorToolkit.cpp
--- a/Engine/Source/Editor/UnrealEd/Private/Toolkits/AssetEditorToolkit.cpp
+++ b/Engine/Source/Editor/UnrealEd/Private/Toolkits/AssetEditorToolkit.cpp
@@ -398,7 +398,7 @@ FText FAssetEditorToolkit::GetToolTipTextForObject(const UObject* InObject)
 
 		FCollectionManagerModule& CollectionManagerModule = FCollectionManagerModule::GetModule();
 
-		const FString CollectionNames = CollectionManagerModule.Get().GetCollectionsStringForObject(*InObject->GetPathName(), ECollectionShareType::CST_All);
+		const FString CollectionNames = CollectionManagerModule.Get().GetCollectionsStringForObject(FSoftObjectPath(InObject), ECollectionShareType::CST_All);
 		if (!CollectionNames.IsEmpty())
 		{
 			ToolTipString += TEXT("\n");
diff --git a/Engine/Source/Runtime/Engine/Private/PrimaryAssetLabel.cpp b/Engine/Source/Runtime/Engine/Private/PrimaryAssetLabel.cpp
--- a/Engine/Source/Runtime/Engine/Private/PrimaryAssetLabel.cpp
+++ b/Engine/Source/Runtime/Engine/Private/PrimaryAssetLabel.cpp
@@ -65,12 +65,9 @@ void UPrimaryAssetLabel::UpdateAssetBundleData()
 	if (AssetCollection.CollectionName != NAME_None)
 	{
 		TArray<FTopLevelAssetPath> NewPaths;
-		TArray<FName> CollectionAssetPathNames;
+		TArray<FSoftObjectPath> CollectionAssets;
 		ICollectionManager& CollectionManager = FCollectionManagerModule::GetModule().Get();
-		CollectionManager.GetAssetsInCollection(AssetCollection.CollectionName, ECollectionShareType::CST_All, CollectionAssetPathNames);
-PRAGMA_DISABLE_DEPRECATION_WARNINGS
-		TArray<FSoftObjectPath> CollectionAssets = UE::SoftObjectPath::Private::ConvertObjectPathNames(CollectionAssetPathNames);
-PRAGMA_ENABLE_DEPRECATION_WARNINGS
+		CollectionManager.GetAssetsInCollection(AssetCollection.CollectionName, ECollectionShareType::CST_All, CollectionAssets);
 		for (int32 Index = 0; Index < CollectionAssets.Num(); ++Index)
 		{
 			FAssetData FoundAsset = Manager.GetAssetRegistry().GetAssetByObjectPath(CollectionAssets[Index]);
