#UE Limit the number of paths reported in changelist descriptions generated by collection checkins to 1000. This can be controlled by setting cvar Collections.MaxCLDescriptionPathCount.

diff --git a/Engine/Source/Developer/CollectionManager/Private/Collection.cpp b/Engine/Source/Developer/CollectionManager/Private/Collection.cpp
--- a/Engine/Source/Developer/CollectionManager/Private/Collection.cpp
+++ b/Engine/Source/Developer/CollectionManager/Private/Collection.cpp
@@ -3,6 +3,7 @@
 #include "Collection.h"
 #include "HAL/PlatformTime.h"
 #include "HAL/FileManager.h"
+#include "HAL/IConsoleManager.h"
 #include "Misc/FileHelper.h"
 #include "Misc/Paths.h"
 #include "Misc/FeedbackContext.h"
@@ -21,6 +22,12 @@
 
 #define LOCTEXT_NAMESPACE "CollectionManager"
 
+static TAutoConsoleVariable<int32> CVarCollectionsMaxCLDescriptionPathCount(
+	TEXT("Collections.MaxCLDescriptionPathCount"),
+	1000,
+	TEXT("Sets the maximum number of paths reported in a changelist when checking in a collection that adds or removes entries."),
+	ECVF_Default);
+
 struct FCollectionUtils
 {
 	static void AppendCollectionToArray(const TSet<FSoftObjectPath>& InObjectSet, TArray<FSoftObjectPath>& OutObjectArray)
@@ -956,6 +963,22 @@ bool FCollection::CheckinCollection(const TArray<FText>& AdditionalChangelistTex
 	{
 		if (StorageMode == ECollectionStorageMode::Static)
 		{
+			auto AddFileListToDescription = [&ChangelistDescBuilder](const TArray<FSoftObjectPath>& Paths)
+			{
+				const int32 MaxPaths = CVarCollectionsMaxCLDescriptionPathCount.GetValueOnAnyThread();
+				const int32 ReportedPaths = FMath::Min(Paths.Num(), MaxPaths);
+				const int32 UnreportedPaths = FMath::Max(0, Paths.Num() - MaxPaths);
+				for (int32 PathIdx = 0; PathIdx < ReportedPaths; ++PathIdx)
+				{
+					const FSoftObjectPath& AddedObjectName = Paths[PathIdx];
+					ChangelistDescBuilder.AppendLine(FText::FromString(AddedObjectName.ToString()));
+				}
+				if (UnreportedPaths > 0)
+				{
+					ChangelistDescBuilder.AppendLineFormat(LOCTEXT("CollectionUnreportedPathsDesc", "... {0} more path(s)"), UnreportedPaths);
+				}
+			};
+
 			// Gather differences from disk
 			TArray<FSoftObjectPath> ObjectsAdded;
 			TArray<FSoftObjectPath> ObjectsRemoved;
@@ -981,10 +1004,7 @@ bool FCollection::CheckinCollection(const TArray<FText>& AdditionalChangelistTex
 				ChangelistDescBuilder.AppendLineFormat(LOCTEXT("CollectionAddedMultipleDesc", "Added {NumberAdded} objects to collection '{CollectionName}':"), Args);
 
 				ChangelistDescBuilder.Indent();
-				for (const FSoftObjectPath& AddedObjectName : ObjectsAdded)
-				{
-					ChangelistDescBuilder.AppendLine(FText::FromString(AddedObjectName.ToString()));
-				}
+				AddFileListToDescription(ObjectsAdded);
 				ChangelistDescBuilder.Unindent();
 			}
 
@@ -997,10 +1017,7 @@ bool FCollection::CheckinCollection(const TArray<FText>& AdditionalChangelistTex
 				ChangelistDescBuilder.AppendLineFormat(LOCTEXT("CollectionRemovedMultipleDesc", "Removed {NumberRemoved} objects from collection '{CollectionName}'"), Args);
 
 				ChangelistDescBuilder.Indent();
-				for (const FSoftObjectPath& RemovedObjectName : ObjectsRemoved)
-				{
-					ChangelistDescBuilder.AppendLine(FText::FromString(RemovedObjectName.ToString()));
-				}
+				AddFileListToDescription(ObjectsRemoved);
 				ChangelistDescBuilder.Unindent();
 			}
 		}
