Add EditorSetting to configure the auto-commit behaviour when collections are modified
* SCC operations can be quite slow in some cases - the auto-checkin behaviour results in a large amount of productivity loss and a stall of the editor waiting for SCC server to respond
* Add an option to disable the auto-commit behaviour if required.  Users would be manually required to commit the .collection file.

diff --git a/Engine/Source/Developer/CollectionManager/Private/Collection.cpp b/Engine/Source/Developer/CollectionManager/Private/Collection.cpp
--- a/Engine/Source/Developer/CollectionManager/Private/Collection.cpp
+++ b/Engine/Source/Developer/CollectionManager/Private/Collection.cpp
@@ -1,6 +1,7 @@
 // Copyright Epic Games, Inc. All Rights Reserved.
 
 #include "Collection.h"
+#include "CollectionConfig.h"
 #include "HAL/PlatformTime.h"
 #include "HAL/FileManager.h"
 #include "HAL/IConsoleManager.h"
@@ -256,7 +257,7 @@ bool FCollection::Save(const TArray<FText>& AdditionalChangelistText, FText& Out
 
 	if ( bSaveSuccessful )
 	{
-		if ( bUseSCC )
+		if ( bUseSCC && GetDefault<UCollectionConfig>()->bAutoCommitOnSave)
 		{
 			// Check in the file if the save was successful
 			if ( bSaveSuccessful )
diff --git a/Engine/Source/Developer/CollectionManager/Private/CollectionConfig.h b/Engine/Source/Developer/CollectionManager/Private/CollectionConfig.h
new file mode 100644
--- /dev/null
+++ b/Engine/Source/Developer/CollectionManager/Private/CollectionConfig.h
@@ -0,0 +1,19 @@
+﻿// Copyright Epic Games, Inc. All Rights Reserved.
+
+#pragma once
+
+#include "UObject/Object.h"
+
+#include "CollectionConfig.generated.h"
+
+#pragma once
+
+UCLASS(minimalapi, config=EditorSettings, meta=(DisplayName = "Collections", ToolTip="Settings to tweak the behaviour of collections in the editor"))
+class UCollectionConfig : public UObject
+{
+	GENERATED_BODY()
+public:
+	/** When enabled, Shared and Private collections will automatically commit their changes to source control */
+	UPROPERTY(config)
+	bool bAutoCommitOnSave = true;
+};
\ No newline at end of file
