Feedback changes to Collections telemetry
* Moved workflow tracking of collection manipulations up one level to target user interaction only
* Added create/destroy collection workflows
* Remove hand-holding utilities aimed to ease verbosity of telemetry event generation

diff --git a/Engine/Source/Developer/AssetTools/Private/CollectionAssetManagement.cpp b/Engine/Source/Developer/AssetTools/Private/CollectionAssetManagement.cpp
--- a/Engine/Source/Developer/AssetTools/Private/CollectionAssetManagement.cpp
+++ b/Engine/Source/Developer/AssetTools/Private/CollectionAssetManagement.cpp
@@ -154,6 +154,11 @@ ECheckBoxState FCollectionAssetManagement::GetCollectionCheckState(FCollectionNa
 	return ECheckBoxState::Unchecked;
 }
 
+int32 FCollectionAssetManagement::GetCurrentAssetCount() const
+{
+	return CurrentAssetPaths.Num();
+}
+
 void FCollectionAssetManagement::UpdateAssetManagementState()
 {
 	FCollectionManagerModule& CollectionManagerModule = FCollectionManagerModule::GetModule();
diff --git a/Engine/Source/Developer/AssetTools/Public/CollectionAssetManagement.h b/Engine/Source/Developer/AssetTools/Public/CollectionAssetManagement.h
--- a/Engine/Source/Developer/AssetTools/Public/CollectionAssetManagement.h
+++ b/Engine/Source/Developer/AssetTools/Public/CollectionAssetManagement.h
@@ -23,6 +23,9 @@ public:
 	/** Set the asset paths that we are currently observing and managing the collection state of */
 	void SetCurrentAssetPaths(const TArray<FSoftObjectPath>& CurrentAssets);
 
+	/** Get the number of assets in the current set */
+	int32 GetCurrentAssetCount() const;
+
 	/** Add the current assets to the given collection */
 	void AddCurrentAssetsToCollection(FCollectionNameType InCollectionKey);
 
diff --git a/Engine/Source/Developer/CollectionManager/Private/CollectionManager.cpp b/Engine/Source/Developer/CollectionManager/Private/CollectionManager.cpp
--- a/Engine/Source/Developer/CollectionManager/Private/CollectionManager.cpp
+++ b/Engine/Source/Developer/CollectionManager/Private/CollectionManager.cpp
@@ -1,7 +1,6 @@
 // Copyright Epic Games, Inc. All Rights Reserved.
 
 #include "CollectionManager.h"
-#include "CollectionManagerTelemetry.h"
 #include "HAL/FileManager.h"
 #include "Misc/Paths.h"
 #include "Containers/Ticker.h"
@@ -1030,10 +1029,7 @@ bool FCollectionManager::AddToCollection(FName CollectionName, ECollectionShareT
 		LastError = LOCTEXT("Error_AddNeedsStaticCollection", "Objects can only be added to static collections.");
 		return false;
 	}
-
-	using namespace UE;
-	auto TelemetryEvent = UE::Analytics::MakeScopedTimedTelemetryEvent<UE::CollectionManager::Analytics::FAddObjects>(FCollectionManagerModule::GetModule().OnAnalyticsEventEmitted);
-
+	
 	int32 NumAdded = 0;
 	for (const FSoftObjectPath& ObjectPath : ObjectPaths)
 	{
@@ -1043,11 +1039,6 @@ bool FCollectionManager::AddToCollection(FName CollectionName, ECollectionShareT
 		}
 	}
 
-	if ((*CollectionRefPtr)->IsUsingSourceControl())
-	{
-		TelemetryEvent.SetStatusCode(TelemetryEvent.GetStatusCode() | CollectionManager::Analytics::StatusCodeFields::UsesSCC);
-	}
-
 	if (NumAdded > 0)
 	{
 		if (InternalSaveCollection(*CollectionRefPtr, LastError))
@@ -1068,7 +1059,6 @@ PRAGMA_DISABLE_DEPRECATION_WARNINGS
 				AssetsAddedEvent.Broadcast(CollectionKey, UE::SoftObjectPath::Private::ConvertSoftObjectPaths(ObjectPaths));
 			}
 PRAGMA_ENABLE_DEPRECATION_WARNINGS
-			TelemetryEvent.SetStatusCode(TelemetryEvent.GetStatusCode() | CollectionManager::Analytics::StatusCodeFields::Success);
 			return true;
 		}
 		else
@@ -1124,10 +1114,7 @@ bool FCollectionManager::RemoveFromCollection(FName CollectionName, ECollectionS
 		LastError = LOCTEXT("Error_RemoveNeedsStaticCollection", "Objects can only be removed from static collections.");
 		return false;
 	}
-
-	using namespace UE;
-	auto TelemetryEvent = Analytics::MakeScopedTimedTelemetryEvent<CollectionManager::Analytics::FRemoveObjects>(FCollectionManagerModule::GetModule().OnAnalyticsEventEmitted);
-
+	
 	TArray<FSoftObjectPath> RemovedAssets;
 	for (const FSoftObjectPath& ObjectPath : ObjectPaths)
 	{
@@ -1137,11 +1124,6 @@ bool FCollectionManager::RemoveFromCollection(FName CollectionName, ECollectionS
 		}
 	}
 
-	if ((*CollectionRefPtr)->IsUsingSourceControl())
-	{
-		TelemetryEvent.SetStatusCode(TelemetryEvent.GetStatusCode() | CollectionManager::Analytics::StatusCodeFields::UsesSCC);
-	}
-
 	if (RemovedAssets.Num() == 0)
 	{
 		// Failed to remove, none of the objects were in the collection
@@ -1167,10 +1149,6 @@ PRAGMA_DISABLE_DEPRECATION_WARNINGS
 			AssetsRemovedEvent.Broadcast(CollectionKey, UE::SoftObjectPath::Private::ConvertSoftObjectPaths(ObjectPaths));
 		}
 PRAGMA_ENABLE_DEPRECATION_WARNINGS
-		if ((*CollectionRefPtr)->IsUsingSourceControl())
-		{
-			TelemetryEvent.SetStatusCode(TelemetryEvent.GetStatusCode() | CollectionManager::Analytics::StatusCodeFields::Success);
-		}
 		return true;
 	}
 	else
diff --git a/Engine/Source/Developer/CollectionManager/Public/CollectionManagerModule.h b/Engine/Source/Developer/CollectionManager/Public/CollectionManagerModule.h
--- a/Engine/Source/Developer/CollectionManager/Public/CollectionManagerModule.h
+++ b/Engine/Source/Developer/CollectionManager/Public/CollectionManagerModule.h
@@ -3,7 +3,6 @@
 #pragma once
 
 #include "CoreMinimal.h"
-#include "CollectionManagerTelemetry.h"
 #include "Modules/ModuleInterface.h"
 #include "Modules/ModuleManager.h"
 #include "UObject/NameTypes.h"
@@ -49,8 +48,6 @@ public:
 		return FModuleManager::Get().IsModuleLoaded(CollectionManagerModuleName);
 	}
 
-	UE::CollectionManager::Analytics::FEmitAnalyticsEvent OnAnalyticsEventEmitted;
-
 private:
 	ICollectionManager* CollectionManager;
 	class FCollectionManagerConsoleCommands* ConsoleCommands;
diff --git a/Engine/Source/Developer/CollectionManager/Public/CollectionManagerTelemetry.h b/Engine/Source/Developer/CollectionManager/Public/CollectionManagerTelemetry.h
deleted file mode 100644
diff --git a/Engine/Source/Editor/ContentBrowser/Private/AssetContextMenu.cpp b/Engine/Source/Editor/ContentBrowser/Private/AssetContextMenu.cpp
--- a/Engine/Source/Editor/ContentBrowser/Private/AssetContextMenu.cpp
+++ b/Engine/Source/Editor/ContentBrowser/Private/AssetContextMenu.cpp
@@ -794,9 +794,13 @@ bool FAssetContextMenu::AddCollectionMenuOptions(UToolMenu* Menu)
 
 		static void OnCollectionClicked(TSharedRef<FCollectionAssetManagement> QuickAssetManagement, FCollectionNameType InCollectionKey)
 		{
+			const double BeginTimeSec = FPlatformTime::Seconds();
+			const int32 ObjectCount = QuickAssetManagement->GetCurrentAssetCount();
+			
 			// The UI actions don't give you the new check state, so we need to emulate the behavior of SCheckBox
 			// Basically, checked will transition to unchecked (removing items), and anything else will transition to checked (adding items)
-			if (GetCollectionCheckState(QuickAssetManagement, InCollectionKey) == ECheckBoxState::Checked)
+			const bool RemoveFromCollection = GetCollectionCheckState(QuickAssetManagement, InCollectionKey) == ECheckBoxState::Checked;
+			if (RemoveFromCollection)
 			{
 				QuickAssetManagement->RemoveCurrentAssetsFromCollection(InCollectionKey);
 			}
@@ -804,6 +808,33 @@ bool FAssetContextMenu::AddCollectionMenuOptions(UToolMenu* Menu)
 			{
 				QuickAssetManagement->AddCurrentAssetsToCollection(InCollectionKey);
 			}
+
+			const double DurationSec = FPlatformTime::Seconds() - BeginTimeSec;
+
+			FContentBrowserModule& ContentBrowserModule = FModuleManager::GetModuleChecked<FContentBrowserModule>( TEXT("ContentBrowser") );
+			if (ContentBrowserModule.GetAssetCollectionTelemetryDelegate().IsBound())
+			{
+				FCollectionTelemetryEvent Event;
+				if (RemoveFromCollection)
+				{
+					FAssetRemovedTelemetryEvent AssetRemoved;
+					AssetRemoved.DurationSec = DurationSec;
+					AssetRemoved.NumRemoved = ObjectCount;
+					AssetRemoved.CollectionShareType = InCollectionKey.Type;
+					AssetRemoved.Workflow = ECollectionTelemetryAssetRemovedWorkflow::ContextMenu;
+					Event = MakeTelemetryEvent(AssetRemoved);
+				}
+				else
+				{
+					FAssetAddedTelemetryEvent AssetAdded;
+					AssetAdded.DurationSec = DurationSec;
+					AssetAdded.NumAdded = ObjectCount;
+					AssetAdded.CollectionShareType = InCollectionKey.Type;
+					AssetAdded.Workflow = ECollectionTelemetryAssetAddedWorkflow::ContextMenu;
+					Event = MakeTelemetryEvent(AssetAdded);
+				}
+				ContentBrowserModule.GetAssetCollectionTelemetryDelegate().Execute(Event);
+			}
 		}
 	};
 
diff --git a/Engine/Source/Editor/ContentBrowser/Private/CollectionContextMenu.cpp b/Engine/Source/Editor/ContentBrowser/Private/CollectionContextMenu.cpp
--- a/Engine/Source/Editor/ContentBrowser/Private/CollectionContextMenu.cpp
+++ b/Engine/Source/Editor/ContentBrowser/Private/CollectionContextMenu.cpp
@@ -451,7 +451,21 @@ void FCollectionContextMenu::ExecuteNewCollection(ECollectionShareType::Type Col
 		return;
 	}
 
+	const double BeginTimeSec = FPlatformTime::Seconds();
+	
 	CollectionView.Pin()->CreateCollectionItem(CollectionType, StorageMode, InCreationPayload);
+
+	// Telemetry Event
+	{
+		FContentBrowserModule& ContentBrowserModule = FModuleManager::GetModuleChecked<FContentBrowserModule>( TEXT("ContentBrowser") );
+		if (ContentBrowserModule.GetAssetCollectionTelemetryDelegate().IsBound())
+		{
+			FCollectionCreatedTelemetryEvent AssetAdded;
+			AssetAdded.DurationSec = FPlatformTime::Seconds() - BeginTimeSec;
+			AssetAdded.CollectionShareType = CollectionType;
+			ContentBrowserModule.GetAssetCollectionTelemetryDelegate().Execute(MakeTelemetryEvent(AssetAdded));
+		}
+	}
 }
 
 void FCollectionContextMenu::ExecuteSetCollectionShareType(ECollectionShareType::Type CollectionType)
@@ -588,7 +602,19 @@ void FCollectionContextMenu::ExecuteDestroyCollection()
 
 FReply FCollectionContextMenu::ExecuteDestroyCollectionConfirmed(TArray<TSharedPtr<FCollectionItem>> CollectionList)
 {
+	const double BeginEventSec = FPlatformTime::Seconds();
+	
 	CollectionView.Pin()->DeleteCollectionItems(CollectionList);
+
+	FContentBrowserModule& ContentBrowserModule = FModuleManager::GetModuleChecked<FContentBrowserModule>( TEXT("ContentBrowser") );
+	if (ContentBrowserModule.GetAssetCollectionTelemetryDelegate().IsBound())
+	{
+		FCollectionsDeletedTelemetryEvent CollectionDeleted;
+		CollectionDeleted.DurationSec = FPlatformTime::Seconds() - BeginEventSec;
+		CollectionDeleted.CollectionsDeleted = CollectionList.Num();
+		ContentBrowserModule.GetAssetCollectionTelemetryDelegate().Execute(MakeTelemetryEvent(CollectionDeleted));
+	}
+	
 	return FReply::Handled();
 }
 
diff --git a/Engine/Source/Editor/ContentBrowser/Private/SCollectionView.cpp b/Engine/Source/Editor/ContentBrowser/Private/SCollectionView.cpp
--- a/Engine/Source/Editor/ContentBrowser/Private/SCollectionView.cpp
+++ b/Engine/Source/Editor/ContentBrowser/Private/SCollectionView.cpp
@@ -1333,6 +1333,7 @@ FReply SCollectionView::HandleDragDropOnCollectionItem(TSharedRef<FCollectionIte
 	}
 	else if (Operation->IsOfType<FAssetDragDropOp>())
 	{
+			
 		TSharedPtr<FAssetDragDropOp> DragDropOp = StaticCastSharedPtr<FAssetDragDropOp>(Operation);
 		const TArray<FAssetData>& DroppedAssets = DragDropOp->GetAssets();
 
@@ -1343,6 +1344,7 @@ FReply SCollectionView::HandleDragDropOnCollectionItem(TSharedRef<FCollectionIte
 			ObjectPaths.Add(AssetData.GetSoftObjectPath());
 		}
 
+		const double BeginTimeSec = FPlatformTime::Seconds();
 		int32 NumAdded = 0;
 		FText Message;
 PRAGMA_DISABLE_DEPRECATION_WARNINGS
@@ -1363,6 +1365,19 @@ PRAGMA_ENABLE_DEPRECATION_WARNINGS
 				Args.Add(TEXT("CollectionName"), FText::FromName(CollectionItem->CollectionName));
 				Message = FText::Format(LOCTEXT("CollectionAssetsAdded", "Added {Number} asset(s) to {CollectionName}"), Args);
 			}
+
+			const double DurationSec = FPlatformTime::Seconds() - BeginTimeSec;
+
+			FContentBrowserModule& ContentBrowserModule = FModuleManager::GetModuleChecked<FContentBrowserModule>( TEXT("ContentBrowser") );
+			if (ContentBrowserModule.GetAssetCollectionTelemetryDelegate().IsBound())
+			{
+				FAssetAddedTelemetryEvent AssetAdded;
+				AssetAdded.DurationSec = DurationSec;
+				AssetAdded.NumAdded = NumAdded;
+				AssetAdded.CollectionShareType = CollectionItem->CollectionType;
+				AssetAdded.Workflow = ECollectionTelemetryAssetAddedWorkflow::DragAndDrop;
+				ContentBrowserModule.GetAssetCollectionTelemetryDelegate().Execute(MakeTelemetryEvent(AssetAdded));
+			}
 		}
 		else
 		{
@@ -1491,6 +1506,8 @@ bool SCollectionView::CollectionNameChangeCommit( const TSharedPtr< FCollectionI
 			NewCollectionParentKey = FCollectionNameType(ParentCollectionItem->CollectionName, ParentCollectionItem->CollectionType);
 		}
 
+		double BeginTimeSec = FPlatformTime::Seconds();
+
 		// If we canceled the name change when creating a new asset, we want to silently remove it
 		if ( !bChangeConfirmed )
 		{
@@ -1528,6 +1545,15 @@ bool SCollectionView::CollectionNameChangeCommit( const TSharedPtr< FCollectionI
 			CollectionItem->OnCollectionCreatedEvent.Execute(FCollectionNameType(NewNameFinal, CollectionItem->CollectionType));
 			CollectionItem->OnCollectionCreatedEvent.Unbind();
 		}
+
+		FContentBrowserModule& ContentBrowserModule = FModuleManager::GetModuleChecked<FContentBrowserModule>( TEXT("ContentBrowser") );
+		if (ContentBrowserModule.GetAssetCollectionTelemetryDelegate().IsBound())
+		{
+			FCollectionCreatedTelemetryEvent CollectionCreatedEvent;
+			CollectionCreatedEvent.DurationSec = FPlatformTime::Seconds() - BeginTimeSec;
+			CollectionCreatedEvent.CollectionShareType = CollectionItem->CollectionType;
+			ContentBrowserModule.GetAssetCollectionTelemetryDelegate().Execute(MakeTelemetryEvent(CollectionCreatedEvent));
+		}
 	}
 	else
 	{
diff --git a/Engine/Source/Editor/ContentBrowser/Public/ContentBrowserModule.h b/Engine/Source/Editor/ContentBrowser/Public/ContentBrowserModule.h
--- a/Engine/Source/Editor/ContentBrowser/Public/ContentBrowserModule.h
+++ b/Engine/Source/Editor/ContentBrowser/Public/ContentBrowserModule.h
@@ -48,6 +48,59 @@ private:
 	friend class FContentBrowserModule;
 };
 
+struct FCollectionTelemetryEventBase {};
+// Workflow event when a collection is created
+struct FCollectionCreatedTelemetryEvent : FCollectionTelemetryEventBase
+{
+	double DurationSec = 0.0;
+	ECollectionShareType::Type CollectionShareType = ECollectionShareType::CST_All;
+};
+
+// Workflow event when a set of collections are deleted
+struct FCollectionsDeletedTelemetryEvent : FCollectionTelemetryEventBase
+{
+	double DurationSec = 0.0;
+	int32 CollectionsDeleted = 0;
+};
+
+enum class ECollectionTelemetryAssetAddedWorkflow : int32
+{
+	ContextMenu = 0,
+	DragAndDrop = 1
+};
+
+// Workflow event when assets are added to a collection
+struct FAssetAddedTelemetryEvent : FCollectionTelemetryEventBase
+{
+	double DurationSec;
+	ECollectionShareType::Type CollectionShareType;
+	uint32 NumAdded;
+	ECollectionTelemetryAssetAddedWorkflow Workflow;
+};
+
+enum class ECollectionTelemetryAssetRemovedWorkflow : int32
+{
+	ContextMenu = 0
+};
+
+// Workflow event when assets are removed from a collection
+struct FAssetRemovedTelemetryEvent : FCollectionTelemetryEventBase
+{
+	double DurationSec;
+	ECollectionShareType::Type CollectionShareType;
+	uint32 NumRemoved;
+	ECollectionTelemetryAssetRemovedWorkflow Workflow;
+};
+
+typedef TVariant<FCollectionCreatedTelemetryEvent, FCollectionsDeletedTelemetryEvent, FAssetAddedTelemetryEvent, FAssetRemovedTelemetryEvent> FCollectionTelemetryEvent;
+
+// Helper to simplify creation of FTelemetryEventTaggedUnion
+template<typename T>
+FCollectionTelemetryEvent MakeTelemetryEvent(const T& Event)
+{
+	return FCollectionTelemetryEvent(TInPlaceType<T>(), Event);
+}
+
 /**
  * Content browser module
  */
@@ -74,7 +127,9 @@ public:
 	DECLARE_DELEGATE_OneParam( FDefaultSelectedPathsDelegate, TArray<FName>& /*VirtualPaths*/ );
 	/** */
 	DECLARE_DELEGATE_OneParam( FDefaultPathsToExpandDelegate, TArray<FName>& /*VirtualPaths*/ );
-
+	/** */
+	DECLARE_DELEGATE_OneParam( FAssetCollectionTelemetryDelegate, const FCollectionTelemetryEvent&);
+	
 	/**
 	 * Called right after the plugin DLL has been loaded and the plugin object has been created
 	 */
@@ -135,6 +190,8 @@ public:
 	/** Override list of paths to expand by default */
 	FDefaultPathsToExpandDelegate& GetDefaultPathsToExpandDelegate() { return DefaultPathsToExpandDelegate; }
 
+	FAssetCollectionTelemetryDelegate& GetAssetCollectionTelemetryDelegate() { return AssetCollectionTelemetryDelegate; }
+
 	FMainMRUFavoritesList* GetRecentlyOpenedAssets() const;
 
 	static const FName NumberOfRecentAssetsName;
@@ -195,4 +252,6 @@ private:
 
 	FDefaultSelectedPathsDelegate DefaultSelectedPathsDelegate;
 	FDefaultPathsToExpandDelegate DefaultPathsToExpandDelegate;
+	
+	FAssetCollectionTelemetryDelegate AssetCollectionTelemetryDelegate;
 };
diff --git a/Engine/Source/Runtime/Analytics/Analytics/Public/AnalyticsTimedTelemetryEvents.h b/Engine/Source/Runtime/Analytics/Analytics/Public/AnalyticsTimedTelemetryEvents.h
deleted file mode 100644
