Expose setting to enable/disable collection auto-commit behaviour
* Changed name CollectionConfig -> CollectionSettings. Settings is the canonical name of these types of objects.
* Derive from UDeveloperSettings to make it auto-discoverable by the SettingsManager

diff --git a/Engine/Source/Developer/CollectionManager/CollectionManager.Build.cs b/Engine/Source/Developer/CollectionManager/CollectionManager.Build.cs
--- a/Engine/Source/Developer/CollectionManager/CollectionManager.Build.cs
+++ b/Engine/Source/Developer/CollectionManager/CollectionManager.Build.cs
@@ -14,7 +14,8 @@ public class CollectionManager : ModuleRules
 				"CoreUObject",
 				"DirectoryWatcher",
 				"Analytics",
-				"SourceControl"
+				"SourceControl",
+				"DeveloperSettings"
 			}
 			);
 	}
diff --git a/Engine/Source/Developer/CollectionManager/Private/Collection.cpp b/Engine/Source/Developer/CollectionManager/Private/Collection.cpp
--- a/Engine/Source/Developer/CollectionManager/Private/Collection.cpp
+++ b/Engine/Source/Developer/CollectionManager/Private/Collection.cpp
@@ -1,7 +1,7 @@
 // Copyright Epic Games, Inc. All Rights Reserved.
 
 #include "Collection.h"
-#include "CollectionConfig.h"
+#include "CollectionSettings.h"
 #include "HAL/PlatformTime.h"
 #include "HAL/FileManager.h"
 #include "HAL/IConsoleManager.h"
@@ -257,7 +257,7 @@ bool FCollection::Save(const TArray<FText>& AdditionalChangelistText, FText& Out
 
 	if ( bSaveSuccessful )
 	{
-		if ( bUseSCC && GetDefault<UCollectionConfig>()->bAutoCommitOnSave)
+		if ( bUseSCC && GetDefault<UCollectionSettings>()->bAutoCommitOnSave)
 		{
 			// Check in the file if the save was successful
 			if ( bSaveSuccessful )
diff --git a/Engine/Source/Developer/CollectionManager/Private/CollectionConfig.h b/Engine/Source/Developer/CollectionManager/Private/CollectionConfig.h
deleted file mode 100644
diff --git a/Engine/Source/Developer/CollectionManager/Private/CollectionSettings.cpp b/Engine/Source/Developer/CollectionManager/Private/CollectionSettings.cpp
new file mode 100644
--- /dev/null
+++ b/Engine/Source/Developer/CollectionManager/Private/CollectionSettings.cpp
@@ -0,0 +1,16 @@
+﻿// Copyright Epic Games, Inc. All Rights Reserved.
+
+#include "CollectionSettings.h"
+#include "Internationalization/Internationalization.h"
+
+#if WITH_EDITOR
+FText UCollectionSettings::GetSectionText() const
+{
+ 	return NSLOCTEXT("Collections", "CollectionsSettingsSection", "Collections");
+}
+
+FText UCollectionSettings::GetSectionDescription() const
+{
+	return NSLOCTEXT("Collections", "CollectionsSettingsDescription", "Settings affecting Asset Collections in Editor");
+}
+#endif
\ No newline at end of file
diff --git a/Engine/Source/Developer/CollectionManager/Public/CollectionSettings.h b/Engine/Source/Developer/CollectionManager/Public/CollectionSettings.h
new file mode 100644
--- /dev/null
+++ b/Engine/Source/Developer/CollectionManager/Public/CollectionSettings.h
@@ -0,0 +1,22 @@
+﻿// Copyright Epic Games, Inc. All Rights Reserved.
+
+#pragma once
+
+#include "Engine/DeveloperSettings.h"
+#include "CollectionSettings.generated.h"
+
+UCLASS(minimalapi, config=EditorSettings)
+class UCollectionSettings : public UDeveloperSettings
+{
+	GENERATED_BODY()
+public:
+
+#if WITH_EDITOR
+	FText GetSectionText() const override;
+	FText GetSectionDescription() const override;
+#endif
+	
+	/** When enabled, Shared and Private collections will automatically commit their changes to source control */
+	UPROPERTY(config, EditAnywhere, Category=Collections, meta = (DisplayName="Auto-Commit to SourceControl on Save"))
+	bool bAutoCommitOnSave = true;
+};
\ No newline at end of file
