Optimize collection manager by allowing batch remove

diff --git a/Engine/Source/Developer/CollectionManager/Private/CollectionManager.cpp b/Engine/Source/Developer/CollectionManager/Private/CollectionManager.cpp
--- a/Engine/Source/Developer/CollectionManager/Private/CollectionManager.cpp
+++ b/Engine/Source/Developer/CollectionManager/Private/CollectionManager.cpp
@@ -1635,22 +1635,36 @@ PRAGMA_ENABLE_DEPRECATION_WARNINGS
 	}
 }
 
-bool FCollectionManager::HandleRedirectorDeleted(const FSoftObjectPath& ObjectPath)
+bool FCollectionManager::HandleRedirectorsDeleted(TConstArrayView<FSoftObjectPath> ObjectPaths)
 {
 	bool bSavedAllCollections = true;
 
 	FTextBuilder AllErrors;
 
 	TArray<FCollectionNameType> UpdatedCollections;
+	TSet<FCollectionNameType> CollectionsToSave;
 
-	// We don't have a cache for on-disk objects, so we have to do this the slower way and query each collection in turn
-	for (const auto& AvailableCollection : AvailableCollections)
+	for (const FSoftObjectPath& ObjectPath : ObjectPaths)
 	{
-		const FCollectionNameType& CollectionKey = AvailableCollection.Key;
-		const TSharedRef<FCollection>& Collection = AvailableCollection.Value;
+		// We don't have a cache for on-disk objects, so we have to do this the slower way and query each collection in turn
+		for (const auto& AvailableCollection : AvailableCollections)
+		{
+			const FCollectionNameType& CollectionKey = AvailableCollection.Key;
+			const TSharedRef<FCollection>& Collection = AvailableCollection.Value;
 
-		if (Collection->IsRedirectorInCollection(ObjectPath))
+			if (Collection->IsRedirectorInCollection(ObjectPath))
+			{
+				CollectionsToSave.Add(CollectionKey);
+			}
+		}
+	}
+
+	for (const FCollectionNameType& CollectionKey : CollectionsToSave)
+	{
+		if (TSharedRef<FCollection>* const CollectionRefPtr = AvailableCollections.Find(CollectionKey))
 		{
+			const TSharedRef<FCollection>& Collection = *CollectionRefPtr;
+
 			FText SaveError;
 			if (InternalSaveCollection(Collection, SaveError))
 			{
@@ -1666,19 +1680,16 @@ bool FCollectionManager::HandleRedirectorDeleted(const FSoftObjectPath& ObjectPa
 		}
 	}
 
-	TArray<FSoftObjectPath> RemovedObjects;
-	RemovedObjects.Add(ObjectPath);
-
 	// Notify every collection that changed
 	for (const FCollectionNameType& UpdatedCollection : UpdatedCollections)
 	{
-		AssetsRemovedFromCollectionDelegate.Broadcast(UpdatedCollection, RemovedObjects);
+		AssetsRemovedFromCollectionDelegate.Broadcast(UpdatedCollection, ObjectPaths);
 	}
 
 PRAGMA_DISABLE_DEPRECATION_WARNINGS
 	if( AssetsRemovedEvent.IsBound())
 	{
-		TArray<FName> RemovedObjectPathNames = UE::SoftObjectPath::Private::ConvertSoftObjectPaths(RemovedObjects);
+		TArray<FName> RemovedObjectPathNames = UE::SoftObjectPath::Private::ConvertSoftObjectPaths(ObjectPaths);
 		for (const FCollectionNameType& UpdatedCollection : UpdatedCollections)
 		{
 			AssetsRemovedEvent.Broadcast(UpdatedCollection, RemovedObjectPathNames);
@@ -1694,6 +1705,11 @@ PRAGMA_ENABLE_DEPRECATION_WARNINGS
 	return bSavedAllCollections;
 }
 
+bool FCollectionManager::HandleRedirectorDeleted(const FSoftObjectPath& ObjectPath)
+{
+	return HandleRedirectorsDeleted(MakeArrayView(&ObjectPath, 1));
+}
+
 void FCollectionManager::HandleObjectRenamed(const FSoftObjectPath& OldObjectPath, const FSoftObjectPath& NewObjectPath)
 {
 	TArray<FCollectionNameType> UpdatedCollections; 
@@ -1731,13 +1747,13 @@ PRAGMA_ENABLE_DEPRECATION_WARNINGS
 	}
 }
 
-void FCollectionManager::HandleObjectDeleted(const FSoftObjectPath& ObjectPath)
+void FCollectionManager::HandleObjectsDeleted(TConstArrayView<FSoftObjectPath> ObjectPaths)
 {
 	TArray<FCollectionNameType> UpdatedCollections;
-	RemoveObjectFromCollections(ObjectPath, UpdatedCollections);
-
-	TArray<FSoftObjectPath> RemovedObjects;
-	RemovedObjects.Add(ObjectPath);
+	for (const FSoftObjectPath& ObjectPath : ObjectPaths)
+	{
+		RemoveObjectFromCollections(ObjectPath, UpdatedCollections);
+	}
 
 	if (UpdatedCollections.Num() > 0)
 	{
@@ -1746,13 +1762,13 @@ void FCollectionManager::HandleObjectDeleted(const FSoftObjectPath& ObjectPath)
 		// Notify every collection that changed
 		for (const FCollectionNameType& UpdatedCollection : UpdatedCollections)
 		{
-			AssetsRemovedFromCollectionDelegate.Broadcast(UpdatedCollection, RemovedObjects);
+			AssetsRemovedFromCollectionDelegate.Broadcast(UpdatedCollection, ObjectPaths);
 		}
 
 PRAGMA_DISABLE_DEPRECATION_WARNINGS
 		if (AssetsRemovedEvent.IsBound())
 		{
-			TArray<FName> RemovedObjectPathNames = UE::SoftObjectPath::Private::ConvertSoftObjectPaths(RemovedObjects);
+			TArray<FName> RemovedObjectPathNames = UE::SoftObjectPath::Private::ConvertSoftObjectPaths(ObjectPaths);
 			for (const FCollectionNameType& UpdatedCollection : UpdatedCollections)
 			{
 				AssetsRemovedEvent.Broadcast(UpdatedCollection, RemovedObjectPathNames);
@@ -1762,6 +1778,11 @@ PRAGMA_ENABLE_DEPRECATION_WARNINGS
 	}
 }
 
+void FCollectionManager::HandleObjectDeleted(const FSoftObjectPath& ObjectPath)
+{
+	HandleObjectsDeleted(MakeArrayView(&ObjectPath, 1));
+}
+
 bool FCollectionManager::TickFileCache(float InDeltaTime)
 {
     QUICK_SCOPE_CYCLE_COUNTER(STAT_FCollectionManager_TickFileCache);
diff --git a/Engine/Source/Developer/CollectionManager/Private/CollectionManager.h b/Engine/Source/Developer/CollectionManager/Private/CollectionManager.h
--- a/Engine/Source/Developer/CollectionManager/Private/CollectionManager.h
+++ b/Engine/Source/Developer/CollectionManager/Private/CollectionManager.h
@@ -255,8 +255,10 @@ PRAGMA_ENABLE_DEPRECATION_WARNINGS
 	virtual FText GetLastError() const override { return LastError; }
 	virtual void HandleFixupRedirectors(ICollectionRedirectorFollower& InRedirectorFollower) override;
 	virtual bool HandleRedirectorDeleted(const FSoftObjectPath& ObjectPath) override;
+	virtual bool HandleRedirectorsDeleted(TConstArrayView<FSoftObjectPath> ObjectPaths) override;
 	virtual void HandleObjectRenamed(const FSoftObjectPath& OldObjectPath, const FSoftObjectPath& NewObjectPath) override;
 	virtual void HandleObjectDeleted(const FSoftObjectPath& ObjectPath) override;
+	virtual void HandleObjectsDeleted(TConstArrayView<FSoftObjectPath> ObjectPaths) override;
 
 	/** Event for when collections are created */
 	DECLARE_DERIVED_EVENT( FCollectionManager, ICollectionManager::FCollectionCreatedEvent, FCollectionCreatedEvent );
diff --git a/Engine/Source/Developer/CollectionManager/Public/ICollectionManager.h b/Engine/Source/Developer/CollectionManager/Public/ICollectionManager.h
--- a/Engine/Source/Developer/CollectionManager/Public/ICollectionManager.h
+++ b/Engine/Source/Developer/CollectionManager/Public/ICollectionManager.h
@@ -309,6 +309,12 @@ public:
 	UE_DEPRECATED(5.1, "FNames containing full object paths are deprecated. Use FSoftObjectPath instead.")
 	virtual bool HandleRedirectorDeleted(const FName& ObjectPath) = 0;
 
+	/**
+	 * Called to notify the collections that redirectors have been deleted and that they should ensure their on-disk representation is re-saved with the fixed up in-memory version
+	 * @return true if all of the collections that were referencing these redirectors could be re-saved, false otherwise
+	 */
+	virtual bool HandleRedirectorsDeleted(TConstArrayView<FSoftObjectPath> ObjectPaths) = 0;
+
 	/** Called to notify the collections that an object has been renamed or moved */
 	virtual void HandleObjectRenamed(const FSoftObjectPath& OldObjectPath, const FSoftObjectPath& NewObjectPath) = 0;
 	UE_DEPRECATED(5.1, "FNames containing full object paths are deprecated. Use FSoftObjectPath instead.")
@@ -319,6 +325,9 @@ public:
 	UE_DEPRECATED(5.1, "FNames containing full object paths are deprecated. Use FSoftObjectPath instead.")
 	virtual void HandleObjectDeleted(const FName& ObjectPath) = 0;
 
+	/** Called to notify the collections that objects have been deleted */
+	virtual void HandleObjectsDeleted(TConstArrayView<FSoftObjectPath> ObjectPaths) = 0;
+
 	/** Event for when collections are created */
 	DECLARE_EVENT_OneParam( ICollectionManager, FCollectionCreatedEvent, const FCollectionNameType& );
 	virtual FCollectionCreatedEvent& OnCollectionCreated() = 0;
diff --git a/Engine/Source/Editor/ContentBrowser/Private/CollectionAssetRegistryBridge.cpp b/Engine/Source/Editor/ContentBrowser/Private/CollectionAssetRegistryBridge.cpp
--- a/Engine/Source/Editor/ContentBrowser/Private/CollectionAssetRegistryBridge.cpp
+++ b/Engine/Source/Editor/ContentBrowser/Private/CollectionAssetRegistryBridge.cpp
@@ -162,20 +162,34 @@ void FCollectionAssetRegistryBridge::OnAssetsRemoved(TConstArrayView<FAssetData>
 {
 	FCollectionManagerModule& CollectionManagerModule = FCollectionManagerModule::GetModule();
 
+	TArray<FSoftObjectPath> RedirectorObjectsRemoved;
+	TArray<FSoftObjectPath> ObjectsRemoved;
+	ObjectsRemoved.Reserve(AssetDatas.Num());
+
 	for (const FAssetData& AssetData : AssetDatas)
 	{
 		if (AssetData.IsRedirector())
 		{
-			// Notify the collections manager that a redirector has been removed
-			// This will attempt to re-save any collections that still have a reference to this redirector in their on-disk collection data
-			CollectionManagerModule.Get().HandleRedirectorDeleted(AssetData.GetSoftObjectPath());
+			RedirectorObjectsRemoved.Add(AssetData.GetSoftObjectPath());
 		}
 		else
 		{
-			// Notify the collections manager that an asset has been removed
-			CollectionManagerModule.Get().HandleObjectDeleted(AssetData.GetSoftObjectPath());
+			ObjectsRemoved.Add(AssetData.GetSoftObjectPath());
 		}
 	}
+
+	// Notify the collections manager that a redirector has been removed
+	// This will attempt to re-save any collections that still have a reference to this redirector in their on-disk collection data
+	if (!RedirectorObjectsRemoved.IsEmpty())
+	{
+		CollectionManagerModule.Get().HandleRedirectorsDeleted(RedirectorObjectsRemoved);
+	}
+
+	// Notify the collections manager that an asset has been removed
+	if (!ObjectsRemoved.IsEmpty())
+	{
+		CollectionManagerModule.Get().HandleObjectsDeleted(ObjectsRemoved);
+	}
 }
 
 #undef LOCTEXT_NAMESPACE
