Update workflows for when not using Collection Setting AutoCommitOnSave
* Only AddToCollection, RemoveFromCollection, SetColor and Reparent will not commit to perforce server if setting is off
* Creating collections notably must commit otherwise it will not be added to a source control working set
* Tunnelled the behaviour of submitting to SCC through the Internal helper function as it does not have the context to decide what workflow is running.

diff --git a/Engine/Source/Developer/CollectionManager/Private/Collection.cpp b/Engine/Source/Developer/CollectionManager/Private/Collection.cpp
--- a/Engine/Source/Developer/CollectionManager/Private/Collection.cpp
+++ b/Engine/Source/Developer/CollectionManager/Private/Collection.cpp
@@ -170,7 +170,7 @@ bool FCollection::Load(FText& OutError)
 	return true;
 }
 
-bool FCollection::Save(const TArray<FText>& AdditionalChangelistText, FText& OutError)
+bool FCollection::Save(const TArray<FText>& AdditionalChangelistText, FText& OutError, bool bForceCommitToRevisionControl)
 {
 	if ( !ensure(SourceFilename.Len()) )
 	{
@@ -257,7 +257,7 @@ bool FCollection::Save(const TArray<FText>& AdditionalChangelistText, FText& Out
 
 	if ( bSaveSuccessful )
 	{
-		if ( bUseSCC && GetDefault<UCollectionSettings>()->bAutoCommitOnSave)
+		if ( bUseSCC && (bForceCommitToRevisionControl || GetDefault<UCollectionSettings>()->bAutoCommitOnSave))
 		{
 			// Check in the file if the save was successful
 			if ( bSaveSuccessful )
diff --git a/Engine/Source/Developer/CollectionManager/Private/Collection.h b/Engine/Source/Developer/CollectionManager/Private/Collection.h
--- a/Engine/Source/Developer/CollectionManager/Private/Collection.h
+++ b/Engine/Source/Developer/CollectionManager/Private/Collection.h
@@ -50,7 +50,7 @@ public:
 	/** Loads content from the SourceFilename into this collection. If false, OutError is a human readable warning depicting the error. */
 	bool Load(FText& OutError);
 	/** Saves this collection to SourceFilename. If submitting to source control, AdditionalChangelistText will be added to the changelist description. If false, OutError is a human readable warning depicting the error. */
-	bool Save(const TArray<FText>& AdditionalChangelistText, FText& OutError);
+	bool Save(const TArray<FText>& AdditionalChangelistText, FText& OutError, bool bForceCommitToRevisionControl);
 	/** Updates this collection to ensure it's the latest version from source control. If false, OutError is a human readable warning depicting the error. */
 	bool Update(FText& OutError);
 	/** Merge the contents of NewCollection into this collection. Returns true if there were changes to merge, or false if the collections were identical. */
diff --git a/Engine/Source/Developer/CollectionManager/Private/CollectionManager.cpp b/Engine/Source/Developer/CollectionManager/Private/CollectionManager.cpp
--- a/Engine/Source/Developer/CollectionManager/Private/CollectionManager.cpp
+++ b/Engine/Source/Developer/CollectionManager/Private/CollectionManager.cpp
@@ -778,7 +778,8 @@ bool FCollectionManager::CreateCollection(FName CollectionName, ECollectionShare
 		return false;
 	}
 
-	if (InternalSaveCollection(NewCollection, LastError))
+	constexpr bool bForceCommitToRevisionControl = true;
+	if (InternalSaveCollection(NewCollection, LastError, bForceCommitToRevisionControl))
 	{
 		CollectionFileCaches[ShareType]->IgnoreNewFile(NewCollection->GetSourceFilename());
 
@@ -829,7 +830,8 @@ bool FCollectionManager::RenameCollection(FName CurrentCollectionName, ECollecti
 			return false;
 		}
 
-		if (!InternalSaveCollection(NewCollection.ToSharedRef(), LastError))
+		bool bForceCommitToRevisionControl = true;
+		if (!InternalSaveCollection(NewCollection.ToSharedRef(), LastError, bForceCommitToRevisionControl))
 		{
 			// Collection failed to save, remove it from the cache
 			RemoveCollection(NewCollection.ToSharedRef(), NewShareType);
@@ -902,8 +904,9 @@ bool FCollectionManager::ReparentCollection(FName CollectionName, ECollectionSha
 		// Does the parent collection need saving in order to have a stable GUID?
 		if ((*ParentCollectionRefPtr)->GetCollectionVersion() < ECollectionVersion::AddedCollectionGuid)
 		{
+			bool bForceCommitToRevisionControl = false;
 			// Try and re-save the parent collection now
-			if (InternalSaveCollection(*ParentCollectionRefPtr, LastError))
+			if (InternalSaveCollection(*ParentCollectionRefPtr, LastError, bForceCommitToRevisionControl))
 			{
 				CollectionFileCaches[ParentShareType]->IgnoreFileModification((*ParentCollectionRefPtr)->GetSourceFilename());
 			}
@@ -931,7 +934,8 @@ bool FCollectionManager::ReparentCollection(FName CollectionName, ECollectionSha
 	(*CollectionRefPtr)->SetParentCollectionGuid(NewParentGuid);
 
 	// Try and save with the new parent GUID
-	if (InternalSaveCollection(*CollectionRefPtr, LastError))
+	bool bForceCommitToRevisionControl = false;
+	if (InternalSaveCollection(*CollectionRefPtr, LastError, bForceCommitToRevisionControl))
 	{
 		CollectionFileCaches[ShareType]->IgnoreFileModification((*CollectionRefPtr)->GetSourceFilename());
 	}
@@ -1040,7 +1044,8 @@ bool FCollectionManager::AddToCollection(FName CollectionName, ECollectionShareT
 
 	if (NumAdded > 0)
 	{
-		if (InternalSaveCollection(*CollectionRefPtr, LastError))
+		constexpr bool bForceCommitToRevisionControl = false;
+		if (InternalSaveCollection(*CollectionRefPtr, LastError, bForceCommitToRevisionControl))
 		{
 			CollectionFileCaches[ShareType]->IgnoreFileModification((*CollectionRefPtr)->GetSourceFilename());
 
@@ -1130,7 +1135,8 @@ bool FCollectionManager::RemoveFromCollection(FName CollectionName, ECollectionS
 		return false;
 	}
 			
-	if (InternalSaveCollection(*CollectionRefPtr, LastError))
+	constexpr bool bForceCommitToRevisionControl = false;
+	if (InternalSaveCollection(*CollectionRefPtr, LastError, bForceCommitToRevisionControl))
 	{
 		CollectionFileCaches[ShareType]->IgnoreFileModification((*CollectionRefPtr)->GetSourceFilename());
 
@@ -1187,7 +1193,8 @@ bool FCollectionManager::SetDynamicQueryText(FName CollectionName, ECollectionSh
 
 	(*CollectionRefPtr)->SetDynamicQueryText(InQueryText);
 	
-	if (InternalSaveCollection(*CollectionRefPtr, LastError))
+	constexpr bool bForceCommitToRevisionControl = true;
+	if (InternalSaveCollection(*CollectionRefPtr, LastError, bForceCommitToRevisionControl))
 	{
 		CollectionFileCaches[ShareType]->IgnoreFileModification((*CollectionRefPtr)->GetSourceFilename());
 
@@ -1281,7 +1288,8 @@ bool FCollectionManager::EmptyCollection(FName CollectionName, ECollectionShareT
 
 	(*CollectionRefPtr)->Empty();
 	
-	if (InternalSaveCollection(*CollectionRefPtr, LastError))
+	constexpr bool bForceCommitToRevisionControl = true;
+	if (InternalSaveCollection(*CollectionRefPtr, LastError, bForceCommitToRevisionControl))
 	{
 		CollectionFileCaches[ShareType]->IgnoreFileModification((*CollectionRefPtr)->GetSourceFilename());
 
@@ -1315,7 +1323,8 @@ bool FCollectionManager::SaveCollection(FName CollectionName, ECollectionShareTy
 			return true;
 		}
 
-		if (InternalSaveCollection(*CollectionRefPtr, LastError))
+		constexpr bool bForceCommitToRevisionControl = true;
+		if (InternalSaveCollection(*CollectionRefPtr, LastError, bForceCommitToRevisionControl))
 		{
 			CollectionFileCaches[ShareType]->IgnoreFileModification((*CollectionRefPtr)->GetSourceFilename());
 
@@ -1437,7 +1446,8 @@ bool FCollectionManager::SetCollectionColor(FName CollectionName, ECollectionSha
 	{
 		(*CollectionRefPtr)->SetCollectionColor(NewColor);
 
-		if (InternalSaveCollection(*CollectionRefPtr, LastError))
+		constexpr bool bForceCommitToRevisionControl = false;
+		if (InternalSaveCollection(*CollectionRefPtr, LastError, bForceCommitToRevisionControl))
 		{
 			CollectionFileCaches[ShareType]->IgnoreFileModification((*CollectionRefPtr)->GetSourceFilename());
 			
@@ -1666,7 +1676,8 @@ bool FCollectionManager::HandleRedirectorsDeleted(TConstArrayView<FSoftObjectPat
 			const TSharedRef<FCollection>& Collection = *CollectionRefPtr;
 
 			FText SaveError;
-			if (InternalSaveCollection(Collection, SaveError))
+			constexpr bool bForceCommitToRevisionControl = true;
+			if (InternalSaveCollection(Collection, SaveError, bForceCommitToRevisionControl))
 			{
 				CollectionFileCaches[CollectionKey.Type]->IgnoreFileModification(Collection->GetSourceFilename());
 
@@ -2069,7 +2080,7 @@ void FCollectionManager::ReplaceObjectInCollections(const FSoftObjectPath& OldOb
 	}
 }
 
-bool FCollectionManager::InternalSaveCollection(const TSharedRef<FCollection>& CollectionRef, FText& OutError)
+bool FCollectionManager::InternalSaveCollection(const TSharedRef<FCollection>& CollectionRef, FText& OutError, bool bForceCommitToRevisionControl)
 {
 	TArray<FText> AdditionalChangelistText;
 
@@ -2093,7 +2104,7 @@ bool FCollectionManager::InternalSaveCollection(const TSharedRef<FCollection>& C
 	}
 
 	// Save the collection
-	return CollectionRef->Save(AdditionalChangelistText, OutError);
+	return CollectionRef->Save(AdditionalChangelistText, OutError, bForceCommitToRevisionControl);
 }
 
 #undef LOCTEXT_NAMESPACE
diff --git a/Engine/Source/Developer/CollectionManager/Private/CollectionManager.h b/Engine/Source/Developer/CollectionManager/Private/CollectionManager.h
--- a/Engine/Source/Developer/CollectionManager/Private/CollectionManager.h
+++ b/Engine/Source/Developer/CollectionManager/Private/CollectionManager.h
@@ -315,8 +315,11 @@ private:
 	/** Replaces an object with another in any collections that contain it */
 	void ReplaceObjectInCollections(const FSoftObjectPath& OldObjectPath, const FSoftObjectPath& NewObjectPath, TArray<FCollectionNameType>& OutUpdatedCollections);
 
-	/** Internal common functionality for saving a collection */
-	bool InternalSaveCollection(const TSharedRef<FCollection>& CollectionRef, FText& OutError);
+	/** Internal common functionality for saving a collection
+	 * bForceCommitToRevisionControl - If the collection's storage mode will save it to source control, then bForceCommitToRevisionControl will ensure that it is committed
+	 * after save.  If this is false, then the collection will be left as a modified file which can be advantageous for slow source control servers.
+	 */
+	bool InternalSaveCollection(const TSharedRef<FCollection>& CollectionRef, FText& OutError, bool bForceCommitToRevisionControl);
 
 private:
 	/** The folders that contain collections */
