Making collection manager thread safe to allow future work to run content browser search in parallel.
Deprecate GetLastError and replace with FText* output parameters on fallible methods.
Use hierarchy of guard types to document access requirements for internal methods.
Separate cache updates from cache access. Cache updates must be performed explicitly so that read->write lock update and possible preemption is clear.

diff --git a/Engine/Plugins/Editor/AssetManagerEditor/Source/AssetManagerEditor/Private/AssetManagerEditorModule.cpp b/Engine/Plugins/Editor/AssetManagerEditor/Source/AssetManagerEditor/Private/AssetManagerEditorModule.cpp
--- a/Engine/Plugins/Editor/AssetManagerEditor/Source/AssetManagerEditor/Private/AssetManagerEditorModule.cpp
+++ b/Engine/Plugins/Editor/AssetManagerEditor/Source/AssetManagerEditor/Private/AssetManagerEditorModule.cpp
@@ -344,7 +344,7 @@ private:
 	void GetPackageDependenciesPerClass(FName SourcePackage, const TArray<FTopLevelAssetPath>& TargetClasses, TArray<FName>& VisitedPackages, TArray<FName>& OutDependentPackages, UE::AssetRegistry::EDependencyQuery RequiredDependencyFlags);
 
 	void LogAssetsWithMultipleLabels();
-	bool CreateOrEmptyCollection(FName CollectionName, ECollectionShareType::Type ShareType);
+	bool CreateOrEmptyCollection(FName CollectionName, ECollectionShareType::Type ShareType, FText& OutError);
 	void WriteProfileFile(const FString& Extension, const FString& FileContents);
 	
 	FString GetSavedAssetRegistryPath(ITargetPlatform* TargetPlatform);
@@ -2439,15 +2439,15 @@ void FAssetManagerEditorModule::DumpAssetDependencies(const TArray<FString>& Arg
 	Manager.WriteCustomReport(FString::Printf(TEXT("PrimaryAssetReferences%s.gv"), *FDateTime::Now().ToString()), ReportLines);
 }
 
-bool FAssetManagerEditorModule::CreateOrEmptyCollection(FName CollectionName, ECollectionShareType::Type ShareType)
+bool FAssetManagerEditorModule::CreateOrEmptyCollection(FName CollectionName, ECollectionShareType::Type ShareType, FText& OutError)
 {
 	ICollectionManager& CollectionManager = FCollectionManagerModule::GetModule().Get();
 
 	if (CollectionManager.CollectionExists(CollectionName, ShareType))
 	{
-		return CollectionManager.EmptyCollection(CollectionName, ShareType);
+		return CollectionManager.EmptyCollection(CollectionName, ShareType, &OutError);
 	}
-	else if (CollectionManager.CreateCollection(CollectionName, ShareType, ECollectionStorageMode::Static))
+	else if (CollectionManager.CreateCollection(CollectionName, ShareType, ECollectionStorageMode::Static, &OutError))
 	{
 		return true;
 	}
@@ -2473,15 +2473,16 @@ bool FAssetManagerEditorModule::WriteCollection(FName CollectionName, ECollectio
 		ObjectPathsToAddToCollection.Add(AssetData.GetSoftObjectPath());
 	}
 
+	FText Error;
 	if (ObjectPathsToAddToCollection.Num() == 0)
 	{
 		UE_LOG(LogAssetManagerEditor, Log, TEXT("Nothing to add to collection %s"), *CollectionName.ToString());
 		ResultsMessage = FText::Format(LOCTEXT("NothingToAddToCollection", "Nothing to add to collection {0}"), FText::FromName(CollectionName));
 	}
-	else if (CreateOrEmptyCollection(CollectionName, ShareType))
+	else if (CreateOrEmptyCollection(CollectionName, ShareType, Error))
 	{
 		if (CollectionManager.AddToCollection(
-				CollectionName, ECollectionShareType::CST_Local, ObjectPathsToAddToCollection.Array()))
+				CollectionName, ECollectionShareType::CST_Local, ObjectPathsToAddToCollection.Array(), nullptr, &Error))
 		{
 			UE_LOG(LogAssetManagerEditor, Log, TEXT("Updated collection %s"), *CollectionName.ToString());
 			ResultsMessage = FText::Format(LOCTEXT("CreateCollectionSucceeded", "Updated collection {0}"), FText::FromName(CollectionName));
@@ -2489,14 +2490,14 @@ bool FAssetManagerEditorModule::WriteCollection(FName CollectionName, ECollectio
 		}
 		else
 		{
-			UE_LOG(LogAssetManagerEditor, Warning, TEXT("Failed to update collection %s. %s"), *CollectionName.ToString(), *CollectionManager.GetLastError().ToString());
-			ResultsMessage = FText::Format(LOCTEXT("AddToCollectionFailed", "Failed to add to collection {0}. {1}"), FText::FromName(CollectionName), CollectionManager.GetLastError());
+			UE_LOG(LogAssetManagerEditor, Warning, TEXT("Failed to update collection %s. %s"), *CollectionName.ToString(), *Error.ToString());
+			ResultsMessage = FText::Format(LOCTEXT("AddToCollectionFailed", "Failed to add to collection {0}. {1}"), FText::FromName(CollectionName), Error);
 		}
 	}
 	else
 	{
-		UE_LOG(LogAssetManagerEditor, Warning, TEXT("Failed to create collection %s. %s"), *CollectionName.ToString(), *CollectionManager.GetLastError().ToString());
-		ResultsMessage = FText::Format(LOCTEXT("CreateCollectionFailed", "Failed to create collection {0}. {1}"), FText::FromName(CollectionName), CollectionManager.GetLastError());
+		UE_LOG(LogAssetManagerEditor, Warning, TEXT("Failed to create collection %s. %s"), *CollectionName.ToString(), *Error.ToString());
+		ResultsMessage = FText::Format(LOCTEXT("CreateCollectionFailed", "Failed to create collection {0}. {1}"), FText::FromName(CollectionName), Error);
 	}
 
 	if (bShowFeedback)
diff --git a/Engine/Plugins/Runtime/AssetTags/Source/AssetTags/Private/AssetTagsSubsystem.cpp b/Engine/Plugins/Runtime/AssetTags/Source/AssetTags/Private/AssetTagsSubsystem.cpp
--- a/Engine/Plugins/Runtime/AssetTags/Source/AssetTags/Private/AssetTagsSubsystem.cpp
+++ b/Engine/Plugins/Runtime/AssetTags/Source/AssetTags/Private/AssetTagsSubsystem.cpp
@@ -19,16 +19,21 @@ DEFINE_LOG_CATEGORY_STATIC(LogAssetTags, Log, All);
 namespace AssetTagsSubsystemUtil
 {
 
-void LogLastCollectionManagerError(ICollectionManager& CollectionManager)
+void LogLastCollectionManagerError(const FText& Error)
 {
-	UE_LOG(LogAssetTags, Warning, TEXT("%s"), *CollectionManager.GetLastError().ToString());
+	UE_LOG(LogAssetTags, Warning, TEXT("Collection manager error: %s"), *Error.ToString());
 }
 
-void LogLastCollectionManagerError()
-{
-	ICollectionManager& CollectionManager = FCollectionManagerModule::GetModule().Get();
-	LogLastCollectionManagerError(CollectionManager);
-}
+// void LogLastCollectionManagerError(ICollectionManager& CollectionManager)
+// {
+// 	UE_LOG(LogAssetTags, Warning, TEXT("%s"), *CollectionManager.GetLastError().ToString());
+// }
+
+// void LogLastCollectionManagerError()
+// {
+// 	ICollectionManager& CollectionManager = FCollectionManagerModule::GetModule().Get();
+// 	LogLastCollectionManagerError(CollectionManager);
+// }
 
 FCollectionNameType FindCollectionByName(ICollectionManager& CollectionManager, const FName Name)
 {
@@ -61,9 +66,10 @@ bool UAssetTagsSubsystem::CreateCollection(const FName Name, const ECollectionSc
 {
 	ICollectionManager& CollectionManager = FCollectionManagerModule::GetModule().Get();
 
-	if (!CollectionManager.IsValidCollectionName(Name.ToString(), ECollectionShareType::CST_All))
+	FText Error;
+	if (!CollectionManager.IsValidCollectionName(Name.ToString(), ECollectionShareType::CST_All, &Error))
 	{
-		AssetTagsSubsystemUtil::LogLastCollectionManagerError(CollectionManager);
+		AssetTagsSubsystemUtil::LogLastCollectionManagerError(Error);
 		return false;
 	}
 
@@ -71,9 +77,9 @@ bool UAssetTagsSubsystem::CreateCollection(const FName Name, const ECollectionSc
 	static_assert((int32)ECollectionShareType::CST_Private == (int32)ECollectionScriptingShareType::Private + 1, "ECollectionShareType::CST_Private is expected to be ECollectionScriptingShareType::Private + 1");
 	static_assert((int32)ECollectionShareType::CST_Shared == (int32)ECollectionScriptingShareType::Shared + 1, "ECollectionShareType::CST_Shared is expected to be ECollectionScriptingShareType::Shared + 1");
 
-	if (!CollectionManager.CreateCollection(Name, (ECollectionShareType::Type)((int32)ShareType + 1), ECollectionStorageMode::Static))
+	if (!CollectionManager.CreateCollection(Name, (ECollectionShareType::Type)((int32)ShareType + 1), ECollectionStorageMode::Static, &Error))
 	{
-		AssetTagsSubsystemUtil::LogLastCollectionManagerError(CollectionManager);
+		AssetTagsSubsystemUtil::LogLastCollectionManagerError(Error);
 		return false;
 	}
 
@@ -90,9 +96,10 @@ bool UAssetTagsSubsystem::DestroyCollection(const FName Name)
 		return false;
 	}
 
-	if (!CollectionManager.DestroyCollection(ResolvedNameAndType.Name, ResolvedNameAndType.Type))
+	FText Error;
+	if (!CollectionManager.DestroyCollection(ResolvedNameAndType.Name, ResolvedNameAndType.Type, &Error))
 	{
-		AssetTagsSubsystemUtil::LogLastCollectionManagerError(CollectionManager);
+		AssetTagsSubsystemUtil::LogLastCollectionManagerError(Error);
 		return false;
 	}
 
@@ -109,15 +116,16 @@ bool UAssetTagsSubsystem::RenameCollection(const FName Name, const FName NewName
 		return false;
 	}
 
-	if (!CollectionManager.IsValidCollectionName(NewName.ToString(), ECollectionShareType::CST_All))
+	FText Error;
+	if (!CollectionManager.IsValidCollectionName(NewName.ToString(), ECollectionShareType::CST_All, &Error))
 	{
-		AssetTagsSubsystemUtil::LogLastCollectionManagerError(CollectionManager);
+		AssetTagsSubsystemUtil::LogLastCollectionManagerError(Error);
 		return false;
 	}
 
-	if (!CollectionManager.RenameCollection(ResolvedNameAndType.Name, ResolvedNameAndType.Type, NewName, ResolvedNameAndType.Type))
+	if (!CollectionManager.RenameCollection(ResolvedNameAndType.Name, ResolvedNameAndType.Type, NewName, ResolvedNameAndType.Type, &Error))
 	{
-		AssetTagsSubsystemUtil::LogLastCollectionManagerError(CollectionManager);
+		AssetTagsSubsystemUtil::LogLastCollectionManagerError(Error);
 		return false;
 	}
 
@@ -140,9 +148,10 @@ bool UAssetTagsSubsystem::ReparentCollection(const FName Name, const FName NewPa
 		return false;
 	}
 
-	if (!CollectionManager.ReparentCollection(ResolvedNameAndType.Name, ResolvedNameAndType.Type, ResolvedParentNameAndType.Name, ResolvedParentNameAndType.Type))
+	FText Error;
+	if (!CollectionManager.ReparentCollection(ResolvedNameAndType.Name, ResolvedNameAndType.Type, ResolvedParentNameAndType.Name, ResolvedParentNameAndType.Type, &Error))
 	{
-		AssetTagsSubsystemUtil::LogLastCollectionManagerError(CollectionManager);
+		AssetTagsSubsystemUtil::LogLastCollectionManagerError(Error);
 		return false;
 	}
 
@@ -159,9 +168,10 @@ bool UAssetTagsSubsystem::EmptyCollection(const FName Name)
 		return false;
 	}
 
-	if (!CollectionManager.EmptyCollection(ResolvedNameAndType.Name, ResolvedNameAndType.Type))
+	FText Error;
+	if (!CollectionManager.EmptyCollection(ResolvedNameAndType.Name, ResolvedNameAndType.Type, &Error))
 	{
-		AssetTagsSubsystemUtil::LogLastCollectionManagerError(CollectionManager);
+		AssetTagsSubsystemUtil::LogLastCollectionManagerError(Error);
 		return false;
 	}
 
@@ -178,9 +188,10 @@ bool UAssetTagsSubsystem::K2_AddAssetToCollection(const FName Name, const FSoftO
 		return false;
 	}
 
-	if (!CollectionManager.AddToCollection(ResolvedNameAndType.Name, ResolvedNameAndType.Type, AssetPath))
+	FText Error;
+	if (!CollectionManager.AddToCollection(ResolvedNameAndType.Name, ResolvedNameAndType.Type, AssetPath, &Error))
 	{
-		AssetTagsSubsystemUtil::LogLastCollectionManagerError(CollectionManager);
+		AssetTagsSubsystemUtil::LogLastCollectionManagerError(Error);
 		return false;
 	}
 
@@ -216,9 +227,10 @@ bool UAssetTagsSubsystem::K2_AddAssetsToCollection(const FName Name, const TArra
 		return false;
 	}
 
-	if (!CollectionManager.AddToCollection(ResolvedNameAndType.Name, ResolvedNameAndType.Type, AssetPaths))
+	FText Error;
+	if (!CollectionManager.AddToCollection(ResolvedNameAndType.Name, ResolvedNameAndType.Type, AssetPaths, nullptr, &Error))
 	{
-		AssetTagsSubsystemUtil::LogLastCollectionManagerError(CollectionManager);
+		AssetTagsSubsystemUtil::LogLastCollectionManagerError(Error);
 		return false;
 	}
 
@@ -264,9 +276,10 @@ bool UAssetTagsSubsystem::K2_RemoveAssetFromCollection(const FName Name, const F
 		return false;
 	}
 
-	if (!CollectionManager.RemoveFromCollection(ResolvedNameAndType.Name, ResolvedNameAndType.Type, AssetPath))
+	FText Error;
+	if (!CollectionManager.RemoveFromCollection(ResolvedNameAndType.Name, ResolvedNameAndType.Type, AssetPath, &Error))
 	{
-		AssetTagsSubsystemUtil::LogLastCollectionManagerError(CollectionManager);
+		AssetTagsSubsystemUtil::LogLastCollectionManagerError(Error);
 		return false;
 	}
 
@@ -302,9 +315,10 @@ bool UAssetTagsSubsystem::K2_RemoveAssetsFromCollection(const FName Name, const
 		return false;
 	}
 
-	if (!CollectionManager.RemoveFromCollection(ResolvedNameAndType.Name, ResolvedNameAndType.Type, AssetPaths))
+	FText Error;
+	if (!CollectionManager.RemoveFromCollection(ResolvedNameAndType.Name, ResolvedNameAndType.Type, AssetPaths, nullptr, &Error))
 	{
-		AssetTagsSubsystemUtil::LogLastCollectionManagerError(CollectionManager);
+		AssetTagsSubsystemUtil::LogLastCollectionManagerError(Error);
 		return false;
 	}
 
diff --git a/Engine/Source/Developer/AssetTools/Private/CollectionAssetManagement.cpp b/Engine/Source/Developer/AssetTools/Private/CollectionAssetManagement.cpp
--- a/Engine/Source/Developer/AssetTools/Private/CollectionAssetManagement.cpp
+++ b/Engine/Source/Developer/AssetTools/Private/CollectionAssetManagement.cpp
@@ -66,7 +66,7 @@ void FCollectionAssetManagement::AddCurrentAssetsToCollection(FCollectionNameTyp
 	bool bSuccess = false;
 	{
 		int32 NumAdded = 0;
-		if (CollectionManagerModule.Get().AddToCollection(InCollectionKey.Name, InCollectionKey.Type, ObjectPaths, &NumAdded))
+		if (CollectionManagerModule.Get().AddToCollection(InCollectionKey.Name, InCollectionKey.Type, ObjectPaths, &NumAdded, &ResultText))
 		{
 			bSuccess = true;
 
@@ -75,10 +75,6 @@ void FCollectionAssetManagement::AddCurrentAssetsToCollection(FCollectionNameTyp
 			Args.Add(TEXT("CollectionName"), FText::FromName(InCollectionKey.Name));
 			ResultText = FText::Format(LOCTEXT("CollectionAssetsAdded", "Added {Number} asset(s) to {CollectionName}"), Args);
 		}
-		else
-		{
-			ResultText = CollectionManagerModule.Get().GetLastError();
-		}
 	}
 
 	if (!ResultText.IsEmpty())
@@ -105,7 +101,7 @@ void FCollectionAssetManagement::RemoveCurrentAssetsFromCollection(FCollectionNa
 	bool bSuccess = false;
 	{
 		int32 NumRemoved = 0;
-		if (CollectionManagerModule.Get().RemoveFromCollection(InCollectionKey.Name, InCollectionKey.Type, ObjectPaths, &NumRemoved))
+		if (CollectionManagerModule.Get().RemoveFromCollection(InCollectionKey.Name, InCollectionKey.Type, ObjectPaths, &NumRemoved, &ResultText))
 		{
 			bSuccess = true;
 
@@ -114,10 +110,6 @@ void FCollectionAssetManagement::RemoveCurrentAssetsFromCollection(FCollectionNa
 			Args.Add(TEXT("CollectionName"), FText::FromName(InCollectionKey.Name));
 			ResultText = FText::Format(LOCTEXT("CollectionAssetsRemoved", "Removed {Number} asset(s) from {CollectionName}"), Args);
 		}
-		else
-		{
-			ResultText = CollectionManagerModule.Get().GetLastError();
-		}
 	}
 
 	if (!ResultText.IsEmpty())
diff --git a/Engine/Source/Developer/CollectionManager/Private/CollectionManager.cpp b/Engine/Source/Developer/CollectionManager/Private/CollectionManager.cpp
--- a/Engine/Source/Developer/CollectionManager/Private/CollectionManager.cpp
+++ b/Engine/Source/Developer/CollectionManager/Private/CollectionManager.cpp
@@ -1,6 +1,9 @@
 // Copyright Epic Games, Inc. All Rights Reserved.
 
 #include "CollectionManager.h"
+
+#include "Algo/Sort.h"
+#include "Algo/Unique.h"
 #include "HAL/FileManager.h"
 #include "Misc/Paths.h"
 #include "Containers/Ticker.h"
@@ -12,215 +15,392 @@
 #include "Async/ParallelFor.h"
 #include "Misc/CommandLine.h"
 #include "SourceControlPreferences.h"
+#include "Tasks/Task.h"
 
 #define LOCTEXT_NAMESPACE "CollectionManager"
 
-FCollectionManagerCache::FCollectionManagerCache(FAvailableCollectionsMap& InAvailableCollections)
+// Base class for lock hierarchy. When used as a function parameter it means the called must hold at least a read lock
+class FCollectionLock
+{
+protected:
+	UE_NODISCARD_CTOR explicit FCollectionLock(FRWLock& InLockObject, bool InWriteLock)
+	: LockObject(InLockObject)
+	, bWriteLock(InWriteLock)
+	{
+		if (InWriteLock)
+		{
+			LockObject.WriteLock();
+		}
+		else
+		{
+			LockObject.ReadLock();
+		}
+	}
+	
+	// Promoted the lock from read to write, possibly being interrupted by another writer in between
+	void PromoteInterruptible()
+	{
+		if(!bWriteLock)
+		{
+			LockObject.ReadUnlock();
+			LockObject.WriteLock();
+			bWriteLock = true;
+		}
+	}
+	
+	~FCollectionLock()
+	{
+		if(bWriteLock)
+		{
+			LockObject.WriteUnlock();
+		}
+		else
+		{
+			LockObject.ReadUnlock();
+		}
+	}
+
+	// Used for assertions to confirm that the correct kind of lock has been taken
+	bool IsWriteLock()
+	{
+		return bWriteLock;
+	}
+	
+private:
+	UE_NONCOPYABLE(FCollectionLock);
+	FRWLock& LockObject;
+	bool bWriteLock = false;
+};
+
+// Scoped lock type used to hold lock and to tag methods which should at least hold a read lock 
+class FCollectionLock_Read : public FCollectionLock
+{
+public:
+	UE_NODISCARD_CTOR explicit FCollectionLock_Read(FRWLock& InLockObject)
+	: FCollectionLock(InLockObject, false)
+	{
+	}
+};
+
+// A lock on the collection manager which begins in a read only state and can be promoted into a write lock with potential interruption in between 
+class FCollectionLock_RW : public FCollectionLock
+{
+public:
+	UE_NODISCARD_CTOR explicit FCollectionLock_RW(FRWLock& InLockObject, bool bWrite = false)
+	: FCollectionLock(InLockObject, bWrite)
+	{
+	}
+	
+	// Promoted the lock from read to write, possibly being interrupted by another writer in between
+	using FCollectionLock::PromoteInterruptible;
+	// Used for assertions to confirm that the correct kind of lock has been taken
+	using FCollectionLock::IsWriteLock;
+};
+
+// Write lock on the collection manager
+class FCollectionLock_Write : public FCollectionLock_RW
+{
+public:
+	UE_NODISCARD_CTOR explicit FCollectionLock_Write(FRWLock& InLockObject)
+	: FCollectionLock_RW(InLockObject, true)
+	{
+	}
+};
+
+/** Wraps up the lazy caching of the collection manager */
+class FCollectionManagerCache
+{
+public:
+	FCollectionManagerCache(TMap<FCollectionNameType, TSharedRef<FCollection>>& InAvailableCollections);
+
+	/** 
+	 * Dirty the parts of the cache that need to change when a collection is added to our collection manager.
+	 * The collection manager must be locked.
+	 */
+	void HandleCollectionAdded(FCollectionLock_Write&);
+	
+	/** 
+	 * Dirty the parts of the cache that need to change when a collection is removed from our collection manager 
+	 * The collection manager must be locked.
+	 */
+	void HandleCollectionRemoved(FCollectionLock_Write&);
+
+	/** 
+	 * Dirty the parts of the cache that need to change when a collection is modified 
+	 * The collcetion manager must be lockedl
+	 */
+	void HandleCollectionChanged(FCollectionLock_Write&);
+
+	/** 
+	 * Update the given dirty parts of the cache based on which parts will be accessed while the given lock is held.
+	 * A read/write lock will be promoted to a write lock if the cache must be updated.
+	 * A write lock may also be passed as it extends the read/write lock.
+	 * The calling thread may be interrupted by another write operation during the promotion operation.
+	 * Therefore, caches should be updated as early as possible in order to prevent invalidation of state.
+	 * 
+	 * This function is used rather than updating the caches in the Get* functions to prevent issues with pre-emption 
+	 * on the lock upgrade deep into a method.
+	 */
+	void UpdateCaches(FCollectionLock_RW& InGuard, ECollectionCacheFlags Flags);
+
+	/** 
+	 * Access the CachedCollectionNamesFromGuids map, asserting that it is up-to-date.
+	 * The collection manager must be read-locked.
+	 */
+	const TMap<FGuid, FCollectionNameType>& GetCachedCollectionNamesFromGuids(FCollectionLock&) const;
+
+	/** 
+	 * Access the CachedObjects map, asserting that it is up-to-date.
+	 * The collection manager must be read-locked.
+	 */
+	const TMap<FSoftObjectPath, TArray<FObjectCollectionInfo>>& GetCachedObjects(FCollectionLock&) const;
+
+	/** 
+	 * Access the CachedHierarchy map, asserting that it is up-to-date.
+	 * The collection manager must be read-locked.
+	 */
+	const TMap<FGuid, TArray<FGuid>>& GetCachedHierarchy(FCollectionLock&) const;
+
+	/** 
+	 * Access the CachedColors array, asserting that it is up-to-date 
+	 * The collection manager must be read-locked.
+	 */
+	const TArray<FLinearColor>& GetCachedColors(FCollectionLock&) const;
+
+	enum class ERecursiveWorkerFlowControl : uint8
+	{
+		Stop,
+		Continue,
+	};
+
+	typedef TFunctionRef<ERecursiveWorkerFlowControl(const FCollectionNameType&, ECollectionRecursionFlags::Flag)> FRecursiveWorkerFunc;
+
+	/** 
+	 * Perform a recursive operation on the given collection and optionally its parents and children.
+	 * The collection manager must be read-locked and UpdateCaches must be called for names and hierarchy.
+	 */
+	void RecursionHelper_DoWork(FCollectionLock&, const FCollectionNameType& InCollectionKey, const ECollectionRecursionFlags::Flags InRecursionMode, FRecursiveWorkerFunc InWorkerFunc) const;
+
+private:
+	/** Reference to the collections that are currently available in our owner collection manager */
+	TMap<FCollectionNameType, TSharedRef<FCollection>>& AvailableCollections;
+
+	/** A map of collection GUIDs to their associated collection names */
+	TMap<FGuid, FCollectionNameType> CachedCollectionNamesFromGuids_Internal;
+
+	/** A map of object paths to their associated collection info - only objects that are in collections will appear in here */
+	TMap<FSoftObjectPath, TArray<FObjectCollectionInfo>> CachedObjects_Internal;
+
+	/** A map of parent collection GUIDs to their child collection GUIDs - only collections that have children will appear in here */
+	TMap<FGuid, TArray<FGuid>> CachedHierarchy_Internal;
+
+	/** An array of all unique colors currently used by collections */
+	TArray<FLinearColor> CachedColors_Internal;
+
+	/** Which parts of the cache are dirty */
+	ECollectionCacheFlags DirtyFlags = ECollectionCacheFlags::All;
+
+	ERecursiveWorkerFlowControl RecursionHelper_DoWorkOnParents(FCollectionLock&, const FCollectionNameType& InCollectionKey, FRecursiveWorkerFunc InWorkerFunc) const;
+	ERecursiveWorkerFlowControl RecursionHelper_DoWorkOnChildren(FCollectionLock&, const FCollectionNameType& InCollectionKey, FRecursiveWorkerFunc InWorkerFunc) const;
+};
+
+FCollectionManagerCache::FCollectionManagerCache(TMap<FCollectionNameType, TSharedRef<FCollection>>& InAvailableCollections)
 	: AvailableCollections(InAvailableCollections)
 {
-	bIsCachedCollectionNamesFromGuidsDirty = true;
-	bIsCachedObjectsDirty = true;
-	bIsCachedHierarchyDirty = true;
-	bIsCachedColorsDirty = true;
 }
 
-void FCollectionManagerCache::HandleCollectionAdded()
+void FCollectionManagerCache::HandleCollectionAdded(FCollectionLock_Write&)
 {
-	bIsCachedCollectionNamesFromGuidsDirty = true;
+	DirtyFlags |= ECollectionCacheFlags::Names;
 }
 
-void FCollectionManagerCache::HandleCollectionRemoved()
+void FCollectionManagerCache::HandleCollectionRemoved(FCollectionLock_Write&)
 {
-	bIsCachedCollectionNamesFromGuidsDirty = true;
-	bIsCachedObjectsDirty = true;
-	bIsCachedHierarchyDirty = true;
-	bIsCachedColorsDirty = true;
+	DirtyFlags |= ECollectionCacheFlags::All;
 }
 
-void FCollectionManagerCache::HandleCollectionChanged()
+void FCollectionManagerCache::HandleCollectionChanged(FCollectionLock_Write&)
 {
-	bIsCachedObjectsDirty = true;
-	bIsCachedHierarchyDirty = true;
-	bIsCachedColorsDirty = true;
+	DirtyFlags |= ECollectionCacheFlags::Objects | ECollectionCacheFlags::Hierarchy | ECollectionCacheFlags::Colors;
 }
 
-const FGuidToCollectionNamesMap& FCollectionManagerCache::GetCachedCollectionNamesFromGuids() const
+void FCollectionManagerCache::UpdateCaches(FCollectionLock_RW& InGuard, ECollectionCacheFlags ToUpdate)
 {
-	if (bIsCachedCollectionNamesFromGuidsDirty)
+	// Updating objects or hierarchy requires name mapping
+	if (EnumHasAnyFlags(ToUpdate, ECollectionCacheFlags::Hierarchy | ECollectionCacheFlags::Objects))
 	{
-		CachedCollectionNamesFromGuids_Internal.Reset();
-		bIsCachedCollectionNamesFromGuidsDirty = false;
+		ToUpdate |= ECollectionCacheFlags::Names;
+	}
+
+	// Updating objects requires hierarchy
+	if (EnumHasAnyFlags(ToUpdate, ECollectionCacheFlags::Objects))
+	{
+		ToUpdate |= ECollectionCacheFlags::Hierarchy;
+	}
+
+	if (EnumHasAnyFlags(DirtyFlags, ToUpdate))
+	{
+		InGuard.PromoteInterruptible();
+	}
+
+	if (!EnumHasAnyFlags(DirtyFlags, ToUpdate))
+	{
+		// Caches we care about were updated while we switched locks
+		return;
+	}
 
-		const double CacheStartTime = FPlatformTime::Seconds();
+	// Limit updates to what's dirty 
+	ToUpdate = ToUpdate & DirtyFlags;
+	const double CacheStartTime = FPlatformTime::Seconds();
 
-		for (const auto& AvailableCollection : AvailableCollections)
+	if (EnumHasAllFlags(ToUpdate, ECollectionCacheFlags::Names))
+	{
+		CachedCollectionNamesFromGuids_Internal.Reset();
+		EnumRemoveFlags(DirtyFlags, ECollectionCacheFlags::Names);
+		for (const TPair<FCollectionNameType, TSharedRef<FCollection>>& AvailableCollection : AvailableCollections)
 		{
 			const FCollectionNameType& CollectionKey = AvailableCollection.Key;
 			const TSharedRef<FCollection>& Collection = AvailableCollection.Value;
 
 			CachedCollectionNamesFromGuids_Internal.Add(Collection->GetCollectionGuid(), CollectionKey);
 		}
-
-		UE_LOG(LogCollectionManager, Verbose, TEXT("Rebuilt the GUID cache for %d collections in %0.6f seconds"), AvailableCollections.Num(), FPlatformTime::Seconds() - CacheStartTime);
 	}
 
-	return CachedCollectionNamesFromGuids_Internal;
-}
-
-const FCollectionObjectsMap& FCollectionManagerCache::GetCachedObjects() const
-{
-	if (bIsCachedObjectsDirty)
+	if (EnumHasAllFlags(ToUpdate, ECollectionCacheFlags::Hierarchy))
 	{
-		CachedObjects_Internal.Reset();
-		bIsCachedObjectsDirty = false;
-
-		const double CacheStartTime = FPlatformTime::Seconds();
+		CachedHierarchy_Internal.Reset();
+		EnumRemoveFlags(DirtyFlags, ECollectionCacheFlags::Hierarchy);
+		const TMap<FGuid, FCollectionNameType>& CachedCollectionNamesFromGuids = GetCachedCollectionNamesFromGuids(InGuard);
 
-		for (const auto& AvailableCollection : AvailableCollections)
+		for (const TPair<FCollectionNameType, TSharedRef<FCollection>>& AvailableCollection : AvailableCollections)
 		{
 			const FCollectionNameType& CollectionKey = AvailableCollection.Key;
 			const TSharedRef<FCollection>& Collection = AvailableCollection.Value;
 
-			const TSet<FSoftObjectPath>& ObjectsInCollection = Collection->GetObjectSet();
-
-			if (ObjectsInCollection.Num() > 0)
+			// Make sure this is a known parent GUID before adding it to the map
+			const FGuid& ParentCollectionGuid = Collection->GetParentCollectionGuid();
+			if (CachedCollectionNamesFromGuids.Contains(ParentCollectionGuid))
 			{
-				auto RebuildCachedObjectsWorker = [&](const FCollectionNameType& InCollectionKey, ECollectionRecursionFlags::Flag InReason) -> ERecursiveWorkerFlowControl
-				{
-					// The worker reason will tell us why this collection is being processed (eg, because it is a parent of the collection we told it to DoWork on),
-					// however, the reason this object exists in that parent collection is because a child collection contains it, and this is the reason we need
-					// to put into the FObjectCollectionInfo, since that's what we'll test against later when we do the "do my children contain this object"? test
-					// That's why we flip the reason logic here...
-					ECollectionRecursionFlags::Flag ReasonObjectInCollection = InReason;
-					switch (InReason)
-					{
-					case ECollectionRecursionFlags::Parents:
-						ReasonObjectInCollection = ECollectionRecursionFlags::Children;
-						break;
-					case ECollectionRecursionFlags::Children:
-						ReasonObjectInCollection = ECollectionRecursionFlags::Parents;
-						break;
-					default:
-						break;
-					}
-
-					for (const FSoftObjectPath& ObjectPath : ObjectsInCollection)
-					{
-						TArray<FObjectCollectionInfo>& ObjectCollectionInfos = CachedObjects_Internal.FindOrAdd(ObjectPath);
-						FObjectCollectionInfo* ObjectInfoPtr = ObjectCollectionInfos.FindByPredicate([&](const FObjectCollectionInfo& InCollectionInfo) { return InCollectionInfo.CollectionKey == InCollectionKey; });
-						if (ObjectInfoPtr)
-						{
-							ObjectInfoPtr->Reason |= ReasonObjectInCollection;
-						}
-						else
-						{
-							ObjectCollectionInfos.Add(FObjectCollectionInfo(InCollectionKey, ReasonObjectInCollection));
-						}
-					}
-					return ERecursiveWorkerFlowControl::Continue;
-				};
-
-				// Recursively process all collections so that they know they contain these objects (and why!)
-				RecursionHelper_DoWork(CollectionKey, ECollectionRecursionFlags::All, RebuildCachedObjectsWorker);
+				TArray<FGuid>& CollectionChildren = CachedHierarchy_Internal.FindOrAdd(ParentCollectionGuid);
+				CollectionChildren.AddUnique(Collection->GetCollectionGuid());
 			}
 		}
-
-		UE_LOG(LogCollectionManager, Verbose, TEXT("Rebuilt the object cache for %d collections in %0.6f seconds (found %d objects)"), AvailableCollections.Num(), FPlatformTime::Seconds() - CacheStartTime, CachedObjects_Internal.Num());
 	}
 
-	return CachedObjects_Internal;
-}
-
-const FCollectionHierarchyMap& FCollectionManagerCache::GetCachedHierarchy() const
-{
-	if (bIsCachedHierarchyDirty)
+	if (EnumHasAllFlags(ToUpdate, ECollectionCacheFlags::Objects))
 	{
-		CachedHierarchy_Internal.Reset();
-		bIsCachedHierarchyDirty = false;
-
-		const FGuidToCollectionNamesMap& CachedCollectionNamesFromGuids = GetCachedCollectionNamesFromGuids();
-
-		const double CacheStartTime = FPlatformTime::Seconds();
+		CachedObjects_Internal.Reset();
+		EnumRemoveFlags(DirtyFlags, ECollectionCacheFlags::Objects);
 
-		for (const auto& AvailableCollection : AvailableCollections)
+		for (const TPair<FCollectionNameType, TSharedRef<FCollection>>& AvailableCollection : AvailableCollections)
 		{
 			const FCollectionNameType& CollectionKey = AvailableCollection.Key;
 			const TSharedRef<FCollection>& Collection = AvailableCollection.Value;
+			const TSet<FSoftObjectPath>& ObjectsInCollection = Collection->GetObjectSet();
 
-			// Make sure this is a known parent GUID before adding it to the map
-			const FGuid& ParentCollectionGuid = Collection->GetParentCollectionGuid();
-			if (CachedCollectionNamesFromGuids.Contains(ParentCollectionGuid))
+			if (ObjectsInCollection.Num() == 0)
 			{
-				auto& CollectionChildren = CachedHierarchy_Internal.FindOrAdd(ParentCollectionGuid);
-				CollectionChildren.AddUnique(Collection->GetCollectionGuid());
+				continue;
 			}
+
+			auto RebuildCachedObjectsWorker = [CachedObjects_Internal=&CachedObjects_Internal, &ObjectsInCollection]
+			(const FCollectionNameType& InCollectionKey, ECollectionRecursionFlags::Flag InReason) -> ERecursiveWorkerFlowControl
+			{
+				// The worker reason will tell us why this collection is being processed (eg, because it is a parent of the collection we told it to DoWork on),
+				// however, the reason this object exists in that parent collection is because a child collection contains it, and this is the reason we need
+				// to put into the FObjectCollectionInfo, since that's what we'll test against later when we do the "do my children contain this object"? test
+				// That's why we flip the reason logic here...
+				ECollectionRecursionFlags::Flag ReasonObjectInCollection = InReason;
+				switch (InReason)
+				{
+				case ECollectionRecursionFlags::Parents:
+					ReasonObjectInCollection = ECollectionRecursionFlags::Children;
+					break;
+				case ECollectionRecursionFlags::Children:
+					ReasonObjectInCollection = ECollectionRecursionFlags::Parents;
+					break;
+				default:
+					break;
+				}
+
+				for (const FSoftObjectPath& ObjectPath : ObjectsInCollection)
+				{
+					TArray<FObjectCollectionInfo>& ObjectCollectionInfos = CachedObjects_Internal->FindOrAdd(ObjectPath);
+					FObjectCollectionInfo* ObjectInfoPtr = ObjectCollectionInfos.FindByPredicate([InCollectionKey](FObjectCollectionInfo& InCollectionInfo) { return InCollectionInfo.CollectionKey == InCollectionKey; });
+					if (ObjectInfoPtr)
+					{
+						ObjectInfoPtr->Reason |= ReasonObjectInCollection;
+					}
+					else
+					{
+						ObjectCollectionInfos.Add(FObjectCollectionInfo(InCollectionKey, ReasonObjectInCollection));
+					}
+				}
+				return ERecursiveWorkerFlowControl::Continue;
+			};
+
+			// Recursively process all collections so that they know they contain these objects (and why!)
+			RecursionHelper_DoWork(InGuard, CollectionKey, ECollectionRecursionFlags::All, RebuildCachedObjectsWorker);
 		}
 
-		UE_LOG(LogCollectionManager, Verbose, TEXT("Rebuilt the hierarchy cache for %d collections in %0.6f seconds"), AvailableCollections.Num(), FPlatformTime::Seconds() - CacheStartTime);
 	}
 
-	return CachedHierarchy_Internal;
-}
-
-const FCollectionColorArray& FCollectionManagerCache::GetCachedColors() const
-{
-	if (bIsCachedColorsDirty)
+	if (EnumHasAllFlags(ToUpdate, ECollectionCacheFlags::Colors))
 	{
 		CachedColors_Internal.Reset();
-		bIsCachedColorsDirty = false;
-
-		const double CacheStartTime = FPlatformTime::Seconds();
-
-		for (const auto& AvailableCollection : AvailableCollections)
+		EnumRemoveFlags(DirtyFlags, ECollectionCacheFlags::Colors);
+		for (const TPair<FCollectionNameType, TSharedRef<FCollection>>& AvailableCollection : AvailableCollections)
 		{
 			const TSharedRef<FCollection>& Collection = AvailableCollection.Value;
 
 			if (const TOptional<FLinearColor> CollectionColor = Collection->GetCollectionColor())
 			{
-				// Only add if not already present (ignores near matches too)
-				const bool bExists = CachedColors_Internal.ContainsByPredicate([CurrentColor = CollectionColor.GetValue()](const FLinearColor& Color) { return CurrentColor.Equals(Color); });
-				if (!bExists)
-				{
-					CachedColors_Internal.Add(CollectionColor.GetValue());
-				}
+				CachedColors_Internal.Add(CollectionColor.GetValue());
 			}
+			// Deduplicate
+			Algo::SortBy(CachedColors_Internal, [](FLinearColor Color) { return GetTypeHash(Color); });
+			CachedColors_Internal.SetNum(Algo::Unique(CachedColors_Internal));
 		}
-
-		UE_LOG(LogCollectionManager, Verbose, TEXT("Rebuilt the color cache for %d collections in %0.6f seconds"), AvailableCollections.Num(), FPlatformTime::Seconds() - CacheStartTime);
 	}
 
-	return CachedColors_Internal;
+	UE_LOG(LogCollectionManager, Verbose, TEXT("Rebuilt caches for %d collections in in %0.6f seconds"), AvailableCollections.Num(), FPlatformTime::Seconds() - CacheStartTime);
 }
 
-void FCollectionManagerCache::RecursionHelper_DoWork(const FCollectionNameType& InCollectionKey, const ECollectionRecursionFlags::Flags InRecursionMode, FRecursiveWorkerFunc InWorkerFunc) const
+void FCollectionManagerCache::RecursionHelper_DoWork(FCollectionLock& Guard,
+	const FCollectionNameType& InCollectionKey,
+	const ECollectionRecursionFlags::Flags InRecursionMode,
+	FRecursiveWorkerFunc InWorkerFunc) const
 {
+	checkf(!EnumHasAnyFlags(DirtyFlags, ECollectionCacheFlags::RecursionWorker), TEXT("Collection cache must be updated with RecursionWorker flags before recursing through hierarchy."));
+
 	if ((InRecursionMode & ECollectionRecursionFlags::Self) && InWorkerFunc(InCollectionKey, ECollectionRecursionFlags::Self) == ERecursiveWorkerFlowControl::Stop)
 	{
 		return;
 	}
 
-	if ((InRecursionMode & ECollectionRecursionFlags::Parents) && RecursionHelper_DoWorkOnParents(InCollectionKey, InWorkerFunc) == ERecursiveWorkerFlowControl::Stop)
+	if ((InRecursionMode & ECollectionRecursionFlags::Parents) && RecursionHelper_DoWorkOnParents(Guard, InCollectionKey, InWorkerFunc) == ERecursiveWorkerFlowControl::Stop)
 	{
 		return;
 	}
 
-	if ((InRecursionMode & ECollectionRecursionFlags::Children) && RecursionHelper_DoWorkOnChildren(InCollectionKey, InWorkerFunc) == ERecursiveWorkerFlowControl::Stop)
+	if ((InRecursionMode & ECollectionRecursionFlags::Children) && RecursionHelper_DoWorkOnChildren(Guard, InCollectionKey, InWorkerFunc) == ERecursiveWorkerFlowControl::Stop)
 	{
 		return;
 	}
 }
 
-FCollectionManagerCache::ERecursiveWorkerFlowControl FCollectionManagerCache::RecursionHelper_DoWorkOnParents(const FCollectionNameType& InCollectionKey, FRecursiveWorkerFunc InWorkerFunc) const
+FCollectionManagerCache::ERecursiveWorkerFlowControl FCollectionManagerCache::RecursionHelper_DoWorkOnParents(
+	FCollectionLock& Guard, const FCollectionNameType& InCollectionKey, FRecursiveWorkerFunc InWorkerFunc) const
 {
 	const TSharedRef<FCollection>* const CollectionRefPtr = AvailableCollections.Find(InCollectionKey);
 	if (CollectionRefPtr)
 	{
-		const FGuidToCollectionNamesMap& CachedCollectionNamesFromGuids = GetCachedCollectionNamesFromGuids();
+		const TMap<FGuid, FCollectionNameType>& CachedCollectionNamesFromGuids = GetCachedCollectionNamesFromGuids(Guard);
 
 		const FCollectionNameType* const ParentCollectionKeyPtr = CachedCollectionNamesFromGuids.Find((*CollectionRefPtr)->GetParentCollectionGuid());
 		if (ParentCollectionKeyPtr)
 		{
-			if (InWorkerFunc(*ParentCollectionKeyPtr, ECollectionRecursionFlags::Parents) == ERecursiveWorkerFlowControl::Stop || RecursionHelper_DoWorkOnParents(*ParentCollectionKeyPtr, InWorkerFunc) == ERecursiveWorkerFlowControl::Stop)
+			if (InWorkerFunc(*ParentCollectionKeyPtr, ECollectionRecursionFlags::Parents) == ERecursiveWorkerFlowControl::Stop || RecursionHelper_DoWorkOnParents(Guard, *ParentCollectionKeyPtr, InWorkerFunc) == ERecursiveWorkerFlowControl::Stop)
 			{
 				return ERecursiveWorkerFlowControl::Stop;
 			}
@@ -230,24 +410,25 @@ FCollectionManagerCache::ERecursiveWorkerFlowControl FCollectionManagerCache::Re
 	return ERecursiveWorkerFlowControl::Continue;
 }
 
-FCollectionManagerCache::ERecursiveWorkerFlowControl FCollectionManagerCache::RecursionHelper_DoWorkOnChildren(const FCollectionNameType& InCollectionKey, FRecursiveWorkerFunc InWorkerFunc) const
+FCollectionManagerCache::ERecursiveWorkerFlowControl FCollectionManagerCache::RecursionHelper_DoWorkOnChildren(
+	FCollectionLock& Guard, const FCollectionNameType& InCollectionKey, FRecursiveWorkerFunc InWorkerFunc) const
 {
 	const TSharedRef<FCollection>* const CollectionRefPtr = AvailableCollections.Find(InCollectionKey);
 	if (CollectionRefPtr)
 	{
-		const FCollectionHierarchyMap& CachedHierarchy = GetCachedHierarchy();
+		const TMap<FGuid, TArray<FGuid>>& CachedHierarchy = GetCachedHierarchy(Guard);
 
 		const TArray<FGuid>* const ChildCollectionGuids = CachedHierarchy.Find((*CollectionRefPtr)->GetCollectionGuid());
 		if (ChildCollectionGuids)
 		{
 			for (const FGuid& ChildCollectionGuid : *ChildCollectionGuids)
 			{
-				const FGuidToCollectionNamesMap& CachedCollectionNamesFromGuids = GetCachedCollectionNamesFromGuids();
+				const TMap<FGuid, FCollectionNameType>& CachedCollectionNamesFromGuids = GetCachedCollectionNamesFromGuids(Guard);
 
 				const FCollectionNameType* const ChildCollectionKeyPtr = CachedCollectionNamesFromGuids.Find(ChildCollectionGuid);
 				if (ChildCollectionKeyPtr)
 				{
-					if (InWorkerFunc(*ChildCollectionKeyPtr, ECollectionRecursionFlags::Children) == ERecursiveWorkerFlowControl::Stop || RecursionHelper_DoWorkOnChildren(*ChildCollectionKeyPtr, InWorkerFunc) == ERecursiveWorkerFlowControl::Stop)
+					if (InWorkerFunc(*ChildCollectionKeyPtr, ECollectionRecursionFlags::Children) == ERecursiveWorkerFlowControl::Stop || RecursionHelper_DoWorkOnChildren(Guard, *ChildCollectionKeyPtr, InWorkerFunc) == ERecursiveWorkerFlowControl::Stop)
 					{
 						return ERecursiveWorkerFlowControl::Stop;
 					}
@@ -259,21 +440,49 @@ FCollectionManagerCache::ERecursiveWorkerFlowControl FCollectionManagerCache::Re
 	return ERecursiveWorkerFlowControl::Continue;
 }
 
-FCollectionManager::FCollectionManager()
-	: CollectionCache(AvailableCollections)
+const TMap<FGuid, FCollectionNameType>& FCollectionManagerCache::GetCachedCollectionNamesFromGuids(FCollectionLock&) const 
+{
+	checkf(!EnumHasAnyFlags(DirtyFlags, ECollectionCacheFlags::Names), TEXT("Accessed guid->name map without updating cache"));
+	return CachedCollectionNamesFromGuids_Internal;
+}
+
+const TMap<FSoftObjectPath, TArray<FObjectCollectionInfo>>& FCollectionManagerCache::GetCachedObjects(FCollectionLock&) const 
+{
+	checkf(!EnumHasAnyFlags(DirtyFlags, ECollectionCacheFlags::Objects), TEXT("Accessd object->collection map without updating cache"));
+	return CachedObjects_Internal;
+}
+
+const TMap<FGuid, TArray<FGuid>>& FCollectionManagerCache::GetCachedHierarchy(FCollectionLock&) const 
+{
+	checkf(!EnumHasAnyFlags(DirtyFlags, ECollectionCacheFlags::Hierarchy), TEXT("Accessed collection hierarchy map without updating cache"));
+	return CachedHierarchy_Internal;
+}
+
+const TArray<FLinearColor>& FCollectionManagerCache::GetCachedColors(FCollectionLock&) const 
 {
-	LastError = LOCTEXT("Error_Unknown", "None");
+	checkf(!EnumHasAnyFlags(DirtyFlags, ECollectionCacheFlags::Colors), TEXT("Accessed collection colors without updating cache"));
+	return CachedColors_Internal;
+}
+
+FStringView FCollectionManager::CollectionExtension = TEXTVIEW("collection");
 
+FCollectionManager::FCollectionManager()
+	: CollectionCache(MakePimpl<FCollectionManagerCache>(AvailableCollections))
+{
 	CollectionFolders[ECollectionShareType::CST_Local] = FPaths::ProjectSavedDir() / TEXT("Collections");
 	CollectionFolders[ECollectionShareType::CST_Private] = FPaths::GameUserDeveloperDir() / TEXT("Collections");
 	CollectionFolders[ECollectionShareType::CST_Shared] = FPaths::ProjectContentDir() / TEXT("Collections");
 
-	CollectionExtension = TEXT("collection");
-
 	bNoFixupRedirectors = FParse::Param(FCommandLine::Get(), TEXT("NoFixupRedirectorsInCollections"));
 
 	LoadCollections();
 
+	// Perform initial caching of collection information ready for user to interact with anything 
+	UE::Tasks::Launch(UE_SOURCE_LOCATION, [this](){
+		FCollectionLock_Write Guard(Lock);
+		CollectionCache->UpdateCaches(Guard, ECollectionCacheFlags::All);
+	});
+
 	// Watch for changes that may happen outside of the collection manager
 	for (int32 CacheIdx = 0; CacheIdx < ECollectionShareType::CST_All; ++CacheIdx)
 	{
@@ -300,22 +509,24 @@ FCollectionManager::FCollectionManager()
 	TickFileCacheDelegateHandle = FTSTicker::GetCoreTicker().AddTicker(FTickerDelegate::CreateRaw(this, &FCollectionManager::TickFileCache), 1.0f);
 }
 
-PRAGMA_DISABLE_DEPRECATION_WARNINGS
 FCollectionManager::~FCollectionManager()
 {
+	FWriteScopeLock Guard(Lock);
 	FTSTicker::GetCoreTicker().RemoveTicker(TickFileCacheDelegateHandle);
 }
-PRAGMA_ENABLE_DEPRECATION_WARNINGS
+
 
 bool FCollectionManager::HasCollections() const
 {
+	FCollectionLock_Read Guard(Lock);
 	return AvailableCollections.Num() > 0;
 }
 
 void FCollectionManager::GetCollections(TArray<FCollectionNameType>& OutCollections) const
 {
+	FCollectionLock_Read Guard(Lock);
 	OutCollections.Reserve(AvailableCollections.Num());
-	for (const auto& AvailableCollection : AvailableCollections)
+	for (const TPair<FCollectionNameType, TSharedRef<FCollection>>& AvailableCollection : AvailableCollections)
 	{
 		const FCollectionNameType& CollectionKey = AvailableCollection.Key;
 		OutCollections.Add(CollectionKey);
@@ -324,6 +535,7 @@ void FCollectionManager::GetCollections(TArray<FCollectionNameType>& OutCollecti
 
 void FCollectionManager::GetCollections(FName CollectionName, TArray<FCollectionNameType>& OutCollections) const
 {
+	FCollectionLock_Read Guard(Lock);
 	for (int32 CacheIdx = 0; CacheIdx < ECollectionShareType::CST_All; ++CacheIdx)
 	{
 		if (AvailableCollections.Contains(FCollectionNameType(CollectionName, ECollectionShareType::Type(CacheIdx))))
@@ -335,7 +547,8 @@ void FCollectionManager::GetCollections(FName CollectionName, TArray<FCollection
 
 void FCollectionManager::GetCollectionNames(ECollectionShareType::Type ShareType, TArray<FName>& CollectionNames) const
 {
-	for (const auto& AvailableCollection : AvailableCollections)
+	FCollectionLock_Read Guard(Lock);
+	for (const TPair<FCollectionNameType, TSharedRef<FCollection>>& AvailableCollection : AvailableCollections)
 	{
 		const FCollectionNameType& CollectionKey = AvailableCollection.Key;
 		if (ShareType == ECollectionShareType::CST_All || ShareType == CollectionKey.Type)
@@ -347,10 +560,12 @@ void FCollectionManager::GetCollectionNames(ECollectionShareType::Type ShareType
 
 void FCollectionManager::GetRootCollections(TArray<FCollectionNameType>& OutCollections) const
 {
-	const FGuidToCollectionNamesMap& CachedCollectionNamesFromGuids = CollectionCache.GetCachedCollectionNamesFromGuids();
+	FCollectionLock_RW Guard(Lock);
+	CollectionCache->UpdateCaches(Guard, ECollectionCacheFlags::Names);
+	const TMap<FGuid, FCollectionNameType>& CachedCollectionNamesFromGuids = CollectionCache->GetCachedCollectionNamesFromGuids(Guard);
 
 	OutCollections.Reserve(AvailableCollections.Num());
-	for (const auto& AvailableCollection : AvailableCollections)
+	for (const TPair<FCollectionNameType, TSharedRef<FCollection>>& AvailableCollection : AvailableCollections)
 	{
 		const FCollectionNameType& CollectionKey = AvailableCollection.Key;
 		const TSharedRef<FCollection>& Collection = AvailableCollection.Value;
@@ -365,9 +580,11 @@ void FCollectionManager::GetRootCollections(TArray<FCollectionNameType>& OutColl
 
 void FCollectionManager::GetRootCollectionNames(ECollectionShareType::Type ShareType, TArray<FName>& CollectionNames) const
 {
-	const FGuidToCollectionNamesMap& CachedCollectionNamesFromGuids = CollectionCache.GetCachedCollectionNamesFromGuids();
+	FCollectionLock_RW Guard(Lock);
+	CollectionCache->UpdateCaches(Guard, ECollectionCacheFlags::Names);
+	const TMap<FGuid, FCollectionNameType>& CachedCollectionNamesFromGuids = CollectionCache->GetCachedCollectionNamesFromGuids(Guard);
 
-	for (const auto& AvailableCollection : AvailableCollections)
+	for (const TPair<FCollectionNameType, TSharedRef<FCollection>>& AvailableCollection : AvailableCollections)
 	{
 		const FCollectionNameType& CollectionKey = AvailableCollection.Key;
 		const TSharedRef<FCollection>& Collection = AvailableCollection.Value;
@@ -385,100 +602,103 @@ void FCollectionManager::GetRootCollectionNames(ECollectionShareType::Type Share
 
 void FCollectionManager::GetChildCollections(FName CollectionName, ECollectionShareType::Type ShareType, TArray<FCollectionNameType>& OutCollections) const
 {
-	const FGuidToCollectionNamesMap& CachedCollectionNamesFromGuids = CollectionCache.GetCachedCollectionNamesFromGuids();
-	const FCollectionHierarchyMap& CachedHierarchy = CollectionCache.GetCachedHierarchy();
+	FCollectionLock_RW Guard(Lock);
+	CollectionCache->UpdateCaches(Guard, ECollectionCacheFlags::Names | ECollectionCacheFlags::Hierarchy);
+
+	const TMap<FGuid, FCollectionNameType>& CachedCollectionNamesFromGuids = CollectionCache->GetCachedCollectionNamesFromGuids(Guard);
+	const TMap<FGuid, TArray<FGuid>>& CachedHierarchy = CollectionCache->GetCachedHierarchy(Guard);
 
-	auto GetChildCollectionsInternal = [&](const FCollectionNameType& InCollectionKey)
+	const int32 Start = ShareType == ECollectionShareType::CST_All ? 0 : int32(ShareType);
+	const int32 End = ShareType == ECollectionShareType::CST_All ? int32(ECollectionShareType::CST_All) : int32(ShareType) + 1;
+
+	for (int32 CacheIdx = Start; CacheIdx < End; ++CacheIdx)
 	{
-		const TSharedRef<FCollection>* const CollectionRefPtr = AvailableCollections.Find(InCollectionKey);
-		if (CollectionRefPtr)
+		FCollectionNameType CollectionKey{ CollectionName, ECollectionShareType::Type(CacheIdx) };
+		const TSharedRef<FCollection>* const CollectionRefPtr = AvailableCollections.Find(CollectionKey);
+		if (!CollectionRefPtr)
 		{
-			const auto* ChildCollectionGuids = CachedHierarchy.Find((*CollectionRefPtr)->GetCollectionGuid());
-			if (ChildCollectionGuids)
-			{
-				for (const FGuid& ChildCollectionGuid : *ChildCollectionGuids)
-				{
-					const FCollectionNameType* const ChildCollectionKeyPtr = CachedCollectionNamesFromGuids.Find(ChildCollectionGuid);
-					if (ChildCollectionKeyPtr)
-					{
-						OutCollections.Add(*ChildCollectionKeyPtr);
-					}
-				}
-			}
+			continue;
 		}
-	};
 
-	if (ShareType == ECollectionShareType::CST_All)
-	{
-		// Asked for all share types, find children in the specified collection name in any cache
-		for (int32 CacheIdx = 0; CacheIdx < ECollectionShareType::CST_All; ++CacheIdx)
+		const TArray<FGuid>* ChildCollectionGuids = CachedHierarchy.Find((*CollectionRefPtr)->GetCollectionGuid());
+		if (!ChildCollectionGuids)
 		{
-			GetChildCollectionsInternal(FCollectionNameType(CollectionName, ECollectionShareType::Type(CacheIdx)));
+			continue;
 		}
-	}
-	else
-	{
-		GetChildCollectionsInternal(FCollectionNameType(CollectionName, ShareType));
-	}
+
+		for (const FGuid& ChildCollectionGuid : *ChildCollectionGuids)
+		{
+			const FCollectionNameType* const ChildCollectionKeyPtr = CachedCollectionNamesFromGuids.Find(ChildCollectionGuid);
+			if (ChildCollectionKeyPtr)
+			{
+				OutCollections.Add(*ChildCollectionKeyPtr);
+			}
+		}
+	};
 }
 
 void FCollectionManager::GetChildCollectionNames(FName CollectionName, ECollectionShareType::Type ShareType, ECollectionShareType::Type ChildShareType, TArray<FName>& CollectionNames) const
 {
-	const FGuidToCollectionNamesMap& CachedCollectionNamesFromGuids = CollectionCache.GetCachedCollectionNamesFromGuids();
-	const FCollectionHierarchyMap& CachedHierarchy = CollectionCache.GetCachedHierarchy();
+	FCollectionLock_RW Guard(Lock);
+	CollectionCache->UpdateCaches(Guard, ECollectionCacheFlags::Names | ECollectionCacheFlags::Hierarchy);
+	const TMap<FGuid, FCollectionNameType>& CachedCollectionNamesFromGuids = CollectionCache->GetCachedCollectionNamesFromGuids(Guard);
+	const TMap<FGuid, TArray<FGuid>>& CachedHierarchy = CollectionCache->GetCachedHierarchy(Guard);
 
-	auto GetChildCollectionsInternal = [&](const FCollectionNameType& InCollectionKey)
+	const int32 Start = ShareType == ECollectionShareType::CST_All ? 0 : int32(ShareType);
+	const int32 End = ShareType == ECollectionShareType::CST_All ? int32(ECollectionShareType::CST_All) : int32(ShareType) + 1;
+
+	for (int32 CacheIdx = Start; CacheIdx < End; ++CacheIdx)
 	{
-		const TSharedRef<FCollection>* const CollectionRefPtr = AvailableCollections.Find(InCollectionKey);
-		if (CollectionRefPtr)
+		FCollectionNameType CollectionKey{ CollectionName, ECollectionShareType::Type(CacheIdx) };
+		const TSharedRef<FCollection>* const CollectionRefPtr = AvailableCollections.Find(CollectionKey);
+		if (!CollectionRefPtr)
 		{
-			const auto* ChildCollectionGuids = CachedHierarchy.Find((*CollectionRefPtr)->GetCollectionGuid());
-			if (ChildCollectionGuids)
-			{
-				for (const FGuid& ChildCollectionGuid : *ChildCollectionGuids)
-				{
-					const FCollectionNameType* const ChildCollectionKeyPtr = CachedCollectionNamesFromGuids.Find(ChildCollectionGuid);
-					if (ChildCollectionKeyPtr && (ChildShareType == ECollectionShareType::CST_All || ChildShareType == ChildCollectionKeyPtr->Type))
-					{
-						CollectionNames.AddUnique(ChildCollectionKeyPtr->Name);
-					}
-				}
-			}
+			continue;
 		}
-	};
 
-	if (ShareType == ECollectionShareType::CST_All)
-	{
-		// Asked for all share types, find children in the specified collection name in any cache
-		for (int32 CacheIdx = 0; CacheIdx < ECollectionShareType::CST_All; ++CacheIdx)
+		const TArray<FGuid>* ChildCollectionGuids = CachedHierarchy.Find((*CollectionRefPtr)->GetCollectionGuid());
+		if (!ChildCollectionGuids)
 		{
-			GetChildCollectionsInternal(FCollectionNameType(CollectionName, ECollectionShareType::Type(CacheIdx)));
+			continue;
 		}
-	}
-	else
-	{
-		GetChildCollectionsInternal(FCollectionNameType(CollectionName, ShareType));
-	}
+
+		for (const FGuid& ChildCollectionGuid : *ChildCollectionGuids)
+		{
+			const FCollectionNameType* const ChildCollectionKeyPtr = CachedCollectionNamesFromGuids.Find(ChildCollectionGuid);
+			if (ChildCollectionKeyPtr && (ChildShareType == ECollectionShareType::CST_All || ChildShareType == ChildCollectionKeyPtr->Type))
+			{
+				CollectionNames.AddUnique(ChildCollectionKeyPtr->Name);
+			}
+		}
+	};
 }
 
 TOptional<FCollectionNameType> FCollectionManager::GetParentCollection(FName CollectionName, ECollectionShareType::Type ShareType) const
 {
+	FCollectionLock_RW Guard(Lock);
 	const TSharedRef<FCollection>* const CollectionRefPtr = AvailableCollections.Find(FCollectionNameType(CollectionName, ShareType));
-	if (CollectionRefPtr)
+	if (!CollectionRefPtr)
 	{
-		const FGuidToCollectionNamesMap& CachedCollectionNamesFromGuids = CollectionCache.GetCachedCollectionNamesFromGuids();
-
-		const FCollectionNameType* const ParentCollectionKeyPtr = CachedCollectionNamesFromGuids.Find((*CollectionRefPtr)->GetParentCollectionGuid());
-		if (ParentCollectionKeyPtr)
-		{
-			return *ParentCollectionKeyPtr;
-		}
+		return TOptional<FCollectionNameType>();
 	}
 
+	CollectionCache->UpdateCaches(Guard, ECollectionCacheFlags::Names);
+	const TMap<FGuid, FCollectionNameType>& CachedCollectionNamesFromGuids = CollectionCache->GetCachedCollectionNamesFromGuids(Guard);
+	const FCollectionNameType* const ParentCollectionKeyPtr = CachedCollectionNamesFromGuids.Find((*CollectionRefPtr)->GetParentCollectionGuid());
+	if (ParentCollectionKeyPtr)
+	{
+		return *ParentCollectionKeyPtr;
+	}
 	return TOptional<FCollectionNameType>();
 }
 
 bool FCollectionManager::CollectionExists(FName CollectionName, ECollectionShareType::Type ShareType) const
+{
+	FCollectionLock_Read Guard(Lock);
+	return CollectionExists_Locked(Guard, CollectionName, ShareType);
+}
+
+bool FCollectionManager::CollectionExists_Locked(FCollectionLock&, FName CollectionName, ECollectionShareType::Type ShareType) const
 {
 	if (ShareType == ECollectionShareType::CST_All)
 	{
@@ -503,11 +723,13 @@ bool FCollectionManager::CollectionExists(FName CollectionName, ECollectionShare
 
 bool FCollectionManager::GetAssetsInCollection(FName CollectionName, ECollectionShareType::Type ShareType, TArray<FSoftObjectPath>& AssetsPaths, ECollectionRecursionFlags::Flags RecursionMode) const
 {
+	FCollectionLock_RW Guard(Lock);
+	CollectionCache->UpdateCaches(Guard, ECollectionCacheFlags::RecursionWorker);
 	bool bFoundAssets = false;
 
-	auto GetAssetsInCollectionWorker = [&](const FCollectionNameType& InCollectionKey, ECollectionRecursionFlags::Flag InReason) -> FCollectionManagerCache::ERecursiveWorkerFlowControl
+	auto GetAssetsInCollectionWorker = [AvailableCollections=&AvailableCollections, &AssetsPaths, &bFoundAssets](const FCollectionNameType& InCollectionKey, ECollectionRecursionFlags::Flag InReason) -> FCollectionManagerCache::ERecursiveWorkerFlowControl
 	{
-		const TSharedRef<FCollection>* const CollectionRefPtr = AvailableCollections.Find(InCollectionKey);
+		const TSharedRef<FCollection>* const CollectionRefPtr = AvailableCollections->Find(InCollectionKey);
 		if (CollectionRefPtr)
 		{
 			(*CollectionRefPtr)->GetAssetsInCollection(AssetsPaths);
@@ -521,12 +743,12 @@ bool FCollectionManager::GetAssetsInCollection(FName CollectionName, ECollection
 		// Asked for all share types, find assets in the specified collection name in any cache
 		for (int32 CacheIdx = 0; CacheIdx < ECollectionShareType::CST_All; ++CacheIdx)
 		{
-			CollectionCache.RecursionHelper_DoWork(FCollectionNameType(CollectionName, ECollectionShareType::Type(CacheIdx)), RecursionMode, GetAssetsInCollectionWorker);
+			CollectionCache->RecursionHelper_DoWork(Guard, FCollectionNameType(CollectionName, ECollectionShareType::Type(CacheIdx)), RecursionMode, GetAssetsInCollectionWorker);
 		}
 	}
 	else
 	{
-		CollectionCache.RecursionHelper_DoWork(FCollectionNameType(CollectionName, ShareType), RecursionMode, GetAssetsInCollectionWorker);
+		CollectionCache->RecursionHelper_DoWork(Guard, FCollectionNameType(CollectionName, ShareType), RecursionMode, GetAssetsInCollectionWorker);
 	}
 
 	return bFoundAssets;
@@ -534,11 +756,13 @@ bool FCollectionManager::GetAssetsInCollection(FName CollectionName, ECollection
 
 bool FCollectionManager::GetClassesInCollection(FName CollectionName, ECollectionShareType::Type ShareType, TArray<FTopLevelAssetPath>& ClassPaths, ECollectionRecursionFlags::Flags RecursionMode) const
 {
+	FCollectionLock_RW Guard(Lock);
+	CollectionCache->UpdateCaches(Guard, ECollectionCacheFlags::RecursionWorker);
 	bool bFoundClasses = false;
 
-	auto GetClassesInCollectionWorker = [&](const FCollectionNameType& InCollectionKey, ECollectionRecursionFlags::Flag InReason) -> FCollectionManagerCache::ERecursiveWorkerFlowControl
+	auto GetClassesInCollectionWorker = [AvailableCollections=&AvailableCollections, &ClassPaths, &bFoundClasses](const FCollectionNameType& InCollectionKey, ECollectionRecursionFlags::Flag InReason) -> FCollectionManagerCache::ERecursiveWorkerFlowControl
 	{
-		const TSharedRef<FCollection>* const CollectionRefPtr = AvailableCollections.Find(InCollectionKey);
+		const TSharedRef<FCollection>* const CollectionRefPtr = AvailableCollections->Find(InCollectionKey);
 		if (CollectionRefPtr)
 		{
 			(*CollectionRefPtr)->GetClassesInCollection(ClassPaths);
@@ -552,12 +776,12 @@ bool FCollectionManager::GetClassesInCollection(FName CollectionName, ECollectio
 		// Asked for all share types, find classes in the specified collection name in any cache
 		for (int32 CacheIdx = 0; CacheIdx < ECollectionShareType::CST_All; ++CacheIdx)
 		{
-			CollectionCache.RecursionHelper_DoWork(FCollectionNameType(CollectionName, ECollectionShareType::Type(CacheIdx)), RecursionMode, GetClassesInCollectionWorker);
+			CollectionCache->RecursionHelper_DoWork(Guard, FCollectionNameType(CollectionName, ECollectionShareType::Type(CacheIdx)), RecursionMode, GetClassesInCollectionWorker);
 		}
 	}
 	else
 	{
-		CollectionCache.RecursionHelper_DoWork(FCollectionNameType(CollectionName, ShareType), RecursionMode, GetClassesInCollectionWorker);
+		CollectionCache->RecursionHelper_DoWork(Guard, FCollectionNameType(CollectionName, ShareType), RecursionMode, GetClassesInCollectionWorker);
 	}
 
 	return bFoundClasses;
@@ -565,11 +789,13 @@ bool FCollectionManager::GetClassesInCollection(FName CollectionName, ECollectio
 
 bool FCollectionManager::GetObjectsInCollection(FName CollectionName, ECollectionShareType::Type ShareType, TArray<FSoftObjectPath>& ObjectPaths, ECollectionRecursionFlags::Flags RecursionMode) const
 {
+	FCollectionLock_RW Guard(Lock);
+	CollectionCache->UpdateCaches(Guard, ECollectionCacheFlags::RecursionWorker);
 	bool bFoundObjects = false;
 
-	auto GetObjectsInCollectionWorker = [&](const FCollectionNameType& InCollectionKey, ECollectionRecursionFlags::Flag InReason) -> FCollectionManagerCache::ERecursiveWorkerFlowControl
+	auto GetObjectsInCollectionWorker = [AvailableCollections=&AvailableCollections, &ObjectPaths, &bFoundObjects](const FCollectionNameType& InCollectionKey, ECollectionRecursionFlags::Flag InReason) -> FCollectionManagerCache::ERecursiveWorkerFlowControl
 	{
-		const TSharedRef<FCollection>* const CollectionRefPtr = AvailableCollections.Find(InCollectionKey);
+		const TSharedRef<FCollection>* const CollectionRefPtr = AvailableCollections->Find(InCollectionKey);
 		if (CollectionRefPtr)
 		{
 			(*CollectionRefPtr)->GetObjectsInCollection(ObjectPaths);
@@ -583,12 +809,12 @@ bool FCollectionManager::GetObjectsInCollection(FName CollectionName, ECollectio
 		// Asked for all share types, find classes in the specified collection name in any cache
 		for (int32 CacheIdx = 0; CacheIdx < ECollectionShareType::CST_All; ++CacheIdx)
 		{
-			CollectionCache.RecursionHelper_DoWork(FCollectionNameType(CollectionName, ECollectionShareType::Type(CacheIdx)), RecursionMode, GetObjectsInCollectionWorker);
+			CollectionCache->RecursionHelper_DoWork(Guard, FCollectionNameType(CollectionName, ECollectionShareType::Type(CacheIdx)), RecursionMode, GetObjectsInCollectionWorker);
 		}
 	}
 	else
 	{
-		CollectionCache.RecursionHelper_DoWork(FCollectionNameType(CollectionName, ShareType), RecursionMode, GetObjectsInCollectionWorker);
+		CollectionCache->RecursionHelper_DoWork(Guard, FCollectionNameType(CollectionName, ShareType), RecursionMode, GetObjectsInCollectionWorker);
 	}
 
 	return bFoundObjects;
@@ -596,9 +822,11 @@ bool FCollectionManager::GetObjectsInCollection(FName CollectionName, ECollectio
 
 void FCollectionManager::GetCollectionsContainingObject(const FSoftObjectPath& ObjectPath, ECollectionShareType::Type ShareType, TArray<FName>& OutCollectionNames, ECollectionRecursionFlags::Flags RecursionMode) const
 {
-	const FCollectionObjectsMap& CachedObjects = CollectionCache.GetCachedObjects();
+	FCollectionLock_RW Guard(Lock);
+	CollectionCache->UpdateCaches(Guard, ECollectionCacheFlags::Objects);
+	const TMap<FSoftObjectPath, TArray<FObjectCollectionInfo>>& CachedObjects = CollectionCache->GetCachedObjects(Guard);
 
-	const auto* ObjectCollectionInfosPtr = CachedObjects.Find(ObjectPath);
+	const TArray<FObjectCollectionInfo>* ObjectCollectionInfosPtr = CachedObjects.Find(ObjectPath);
 	if (ObjectCollectionInfosPtr)
 	{
 		for (const FObjectCollectionInfo& ObjectCollectionInfo : *ObjectCollectionInfosPtr)
@@ -613,9 +841,11 @@ void FCollectionManager::GetCollectionsContainingObject(const FSoftObjectPath& O
 
 void FCollectionManager::GetCollectionsContainingObject(const FSoftObjectPath& ObjectPath, TArray<FCollectionNameType>& OutCollections, ECollectionRecursionFlags::Flags RecursionMode) const
 {
-	const FCollectionObjectsMap& CachedObjects = CollectionCache.GetCachedObjects();
+	FCollectionLock_RW Guard(Lock);
+	CollectionCache->UpdateCaches(Guard, ECollectionCacheFlags::Objects);
+	const TMap<FSoftObjectPath, TArray<FObjectCollectionInfo>>& CachedObjects = CollectionCache->GetCachedObjects(Guard);
 
-	const auto* ObjectCollectionInfosPtr = CachedObjects.Find(ObjectPath);
+	const TArray<FObjectCollectionInfo>* ObjectCollectionInfosPtr = CachedObjects.Find(ObjectPath);
 	if (ObjectCollectionInfosPtr)
 	{
 		OutCollections.Reserve(OutCollections.Num() + ObjectCollectionInfosPtr->Num());
@@ -631,11 +861,13 @@ void FCollectionManager::GetCollectionsContainingObject(const FSoftObjectPath& O
 
 void FCollectionManager::GetCollectionsContainingObjects(const TArray<FSoftObjectPath>& ObjectPaths, TMap<FCollectionNameType, TArray<FSoftObjectPath>>& OutCollectionsAndMatchedObjects, ECollectionRecursionFlags::Flags RecursionMode) const
 {
-	const FCollectionObjectsMap& CachedObjects = CollectionCache.GetCachedObjects();
+	FCollectionLock_RW Guard(Lock);
+	CollectionCache->UpdateCaches(Guard, ECollectionCacheFlags::Objects);
+	const TMap<FSoftObjectPath, TArray<FObjectCollectionInfo>>& CachedObjects = CollectionCache->GetCachedObjects(Guard);
 
 	for (const FSoftObjectPath& ObjectPath : ObjectPaths)
 	{
-		const auto* ObjectCollectionInfosPtr = CachedObjects.Find(ObjectPath);
+		const TArray<FObjectCollectionInfo>* ObjectCollectionInfosPtr = CachedObjects.Find(ObjectPath);
 		if (ObjectCollectionInfosPtr)
 		{
 			for (const FObjectCollectionInfo& ObjectCollectionInfo : *ObjectCollectionInfosPtr)
@@ -652,16 +884,17 @@ void FCollectionManager::GetCollectionsContainingObjects(const TArray<FSoftObjec
 
 FString FCollectionManager::GetCollectionsStringForObject(const FSoftObjectPath& ObjectPath, ECollectionShareType::Type ShareType, ECollectionRecursionFlags::Flags RecursionMode, bool bFullPaths) const
 {
-	const FCollectionObjectsMap& CachedObjects = CollectionCache.GetCachedObjects();
+	FCollectionLock_RW Guard(Lock);
+	CollectionCache->UpdateCaches(Guard, ECollectionCacheFlags::Objects | ECollectionCacheFlags::RecursionWorker);
+	const TMap<FSoftObjectPath, TArray<FObjectCollectionInfo>>& CachedObjects = CollectionCache->GetCachedObjects(Guard);
 
-	const auto* ObjectCollectionInfosPtr = CachedObjects.Find(ObjectPath);
+	const TArray<FObjectCollectionInfo>* ObjectCollectionInfosPtr = CachedObjects.Find(ObjectPath);
 	if (ObjectCollectionInfosPtr)
 	{
 		TArray<FString> CollectionNameStrings;
-
 		TArray<FString> CollectionPathStrings;
 
-		auto GetCollectionsStringForObjectWorker = [&](const FCollectionNameType& InCollectionKey, ECollectionRecursionFlags::Flag InReason) -> FCollectionManagerCache::ERecursiveWorkerFlowControl
+		auto GetCollectionsStringForObjectWorker = [&CollectionPathStrings](const FCollectionNameType& InCollectionKey, ECollectionRecursionFlags::Flag InReason) -> FCollectionManagerCache::ERecursiveWorkerFlowControl
 		{
 			CollectionPathStrings.Insert(InCollectionKey.Name.ToString(), 0);
 			return FCollectionManagerCache::ERecursiveWorkerFlowControl::Continue;
@@ -674,7 +907,7 @@ FString FCollectionManager::GetCollectionsStringForObject(const FSoftObjectPath&
 				if (bFullPaths)
 				{
 					CollectionPathStrings.Reset();
-					CollectionCache.RecursionHelper_DoWork(ObjectCollectionInfo.CollectionKey, ECollectionRecursionFlags::SelfAndParents, GetCollectionsStringForObjectWorker);
+					CollectionCache->RecursionHelper_DoWork(Guard, ObjectCollectionInfo.CollectionKey, ECollectionRecursionFlags::SelfAndParents, GetCollectionsStringForObjectWorker);
 					CollectionNameStrings.Add(FString::Join(CollectionPathStrings, TEXT("/")));
 				}
 				else
@@ -696,6 +929,7 @@ FString FCollectionManager::GetCollectionsStringForObject(const FSoftObjectPath&
 
 void FCollectionManager::CreateUniqueCollectionName(const FName& BaseName, ECollectionShareType::Type ShareType, FName& OutCollectionName) const
 {
+	FCollectionLock_Read Guard(Lock);
 	int32 IntSuffix = 1;
 	bool CollectionAlreadyExists = false;
 	do
@@ -709,19 +943,22 @@ void FCollectionManager::CreateUniqueCollectionName(const FName& BaseName, EColl
 			OutCollectionName = *FString::Printf(TEXT("%s%d"), *BaseName.ToString(), IntSuffix);
 		}
 
-		CollectionAlreadyExists = CollectionExists(OutCollectionName, ShareType);
+		CollectionAlreadyExists = CollectionExists_Locked(Guard, OutCollectionName, ShareType);
 		++IntSuffix;
 	}
 	while (CollectionAlreadyExists);
 }
 
-bool FCollectionManager::IsValidCollectionName(const FString& CollectionName, ECollectionShareType::Type ShareType) const
+bool FCollectionManager::IsValidCollectionName(const FString& CollectionName, ECollectionShareType::Type ShareType, FText* OutError) const
 {
 	// Make sure we are not creating an FName that is too large
 	if (CollectionName.Len() >= NAME_SIZE)
 	{
-		LastError = FText::Format(LOCTEXT("Error_CollectionNameTooLong", "This collection name is too long ({0} characters), the maximum is {1}. Please choose a shorter name. Collection name: {2}"),
-			FText::AsNumber(CollectionName.Len()), FText::AsNumber(NAME_SIZE), FText::FromString(CollectionName));
+		if (OutError)
+		{
+			*OutError = FText::Format(LOCTEXT("Error_CollectionNameTooLong","This collection name is too long ({0} characters), the maximum is {1}. Please choose a shorter name. Collection name: {2}"),
+				FText::AsNumber(CollectionName.Len()), FText::AsNumber(NAME_SIZE), FText::FromString(CollectionName));
+		}
 		return false;
 	}
 
@@ -730,281 +967,340 @@ bool FCollectionManager::IsValidCollectionName(const FString& CollectionName, EC
 	// Make sure the we actually have a new name set
 	if (CollectionNameFinal.IsNone())
 	{
-		LastError = LOCTEXT("Error_CollectionNameEmptyOrNone", "This collection name cannot be empty or 'None'.");
+		if (OutError)
+		{
+			*OutError = LOCTEXT("Error_CollectionNameEmptyOrNone", "This collection name cannot be empty or 'None'.");
+		}
 		return false;
 	}
 
 	// Make sure the new name only contains valid characters
-	if (!CollectionNameFinal.IsValidXName(INVALID_OBJECTNAME_CHARACTERS INVALID_LONGPACKAGE_CHARACTERS, &LastError))
+	if (!CollectionNameFinal.IsValidXName(INVALID_OBJECTNAME_CHARACTERS INVALID_LONGPACKAGE_CHARACTERS, &*OutError))
 	{
 		return false;
 	}
 
 	// Make sure we're not duplicating an existing collection name
+	// NB: Ok to call public function here because we don't need acquire a lock for the previous checks 
 	if (CollectionExists(CollectionNameFinal, ShareType))
 	{
-		LastError = FText::Format(LOCTEXT("Error_CollectionAlreadyExists", "A collection already exists with the name '{0}'."), FText::FromName(CollectionNameFinal));
+		if (OutError)
+		{
+			*OutError = FText::Format(LOCTEXT("Error_CollectionAlreadyExists", "A collection already exists with the name '{0}'."), FText::FromName(CollectionNameFinal));
+		}
 		return false;
 	}
 
 	return true;
 }
 
-bool FCollectionManager::CreateCollection(FName CollectionName, ECollectionShareType::Type ShareType, ECollectionStorageMode::Type StorageMode)
+bool FCollectionManager::CreateCollection(FName CollectionName, ECollectionShareType::Type ShareType, ECollectionStorageMode::Type StorageMode, FText* OutError)
 {
 	if (!ensure(ShareType < ECollectionShareType::CST_All))
 	{
 		// Bad share type
-		LastError = LOCTEXT("Error_Internal", "There was an internal error.");
+		if (OutError)
+		{
+			*OutError = LOCTEXT("Error_Internal", "There was an internal error.");
+		}
 		return false;
 	}
 
-	// Try to add the collection
-	const bool bUseSCC = ShouldUseSCC(ShareType);
-	const FString CollectionFilename = GetCollectionFilename(CollectionName, ShareType);
-
-	// Validate collection name as file name
-	bool bFilenameValid = FFileHelper::IsFilenameValidForSaving(CollectionName.ToString(), LastError);
-	if (!bFilenameValid)
+	if (!IsValidCollectionName(CollectionName.ToString(), ShareType, OutError))
 	{
 		return false;
 	}
 
-	TSharedRef<FCollection> NewCollection = MakeShareable(new FCollection(CollectionFilename, bUseSCC, StorageMode));
-	if (!AddCollection(NewCollection, ShareType))
 	{
-		// Failed to add the collection, it already exists
-		LastError = LOCTEXT("Error_AlreadyExists", "The collection already exists.");
-		return false;
-	}
+		FCollectionLock_Write Guard(Lock);
 
-	constexpr bool bForceCommitToRevisionControl = true;
-	if (InternalSaveCollection(NewCollection, LastError, bForceCommitToRevisionControl))
-	{
-		CollectionFileCaches[ShareType]->IgnoreNewFile(NewCollection->GetSourceFilename());
+		// Try to add the collection
+		const bool bUseSCC = ShouldUseSCC(ShareType);
+		const FString CollectionFilename = GetCollectionFilename(CollectionName, ShareType);
 
-		// Collection saved!
-		CollectionCreatedEvent.Broadcast(FCollectionNameType(CollectionName, ShareType));
-		return true;
-	}
-	else
-	{
-		// Collection failed to save, remove it from the cache
-		RemoveCollection(NewCollection, ShareType);
-		return false;
-	}
-}
+		// Validate collection name as file name
+		FText UnusedError;
+		bool bFilenameValid = FFileHelper::IsFilenameValidForSaving(CollectionName.ToString(), OutError ? *OutError : UnusedError);
+		if (!bFilenameValid)
+		{
+			return false;
+		}
 
-bool FCollectionManager::RenameCollection(FName CurrentCollectionName, ECollectionShareType::Type CurrentShareType, FName NewCollectionName, ECollectionShareType::Type NewShareType)
-{
-	if (!ensure(CurrentShareType < ECollectionShareType::CST_All) || !ensure(NewShareType < ECollectionShareType::CST_All))
-	{
-		// Bad share type
-		LastError = LOCTEXT("Error_Internal", "There was an internal error.");
-		return false;
-	}
+		TSharedRef<FCollection> NewCollection = MakeShareable(new FCollection(CollectionFilename, bUseSCC, StorageMode));
+		if (!AddCollection(Guard, NewCollection, ShareType))
+		{
+			// Failed to add the collection, it already exists
+			if (OutError)
+			{
+				*OutError = LOCTEXT("Error_AlreadyExists", "The collection already exists.");
+			}
+			return false;
+		}
 
-	const FCollectionNameType OriginalCollectionKey(CurrentCollectionName, CurrentShareType);
-	const FCollectionNameType NewCollectionKey(NewCollectionName, NewShareType);
+		constexpr bool bForceCommitToRevisionControl = true;
+		if (!InternalSaveCollection(Guard, NewCollection, OutError, bForceCommitToRevisionControl))
+		{
+			// Collection failed to save, remove it from the cache
+			RemoveCollection(Guard, NewCollection, ShareType);
+			return false;
+		}
 
-	TSharedRef<FCollection>* const CollectionRefPtr = AvailableCollections.Find(OriginalCollectionKey);
-	if (!CollectionRefPtr)
+		CollectionFileCaches[ShareType]->IgnoreNewFile(NewCollection->GetSourceFilename());
+	}
+
+	// Collection saved!
+	// Broadcast events outside of lock 
+	CollectionCreatedEvent.Broadcast(FCollectionNameType(CollectionName, ShareType));
+	return true;
+}
+
+bool FCollectionManager::RenameCollection(FName CurrentCollectionName, ECollectionShareType::Type CurrentShareType, FName NewCollectionName, ECollectionShareType::Type NewShareType, FText* OutError)
+{
+	if (!ensure(CurrentShareType < ECollectionShareType::CST_All) || !ensure(NewShareType < ECollectionShareType::CST_All))
 	{
-		// The collection doesn't exist
-		LastError = LOCTEXT("Error_DoesntExist", "The collection doesn't exist.");
+		// Bad share type
+		if (OutError)
+		{
+			*OutError = LOCTEXT("Error_Internal", "There was an internal error.");
+		}
 		return false;
 	}
 
-	// Add the new collection
-	TSharedPtr<FCollection> NewCollection;
+	const FCollectionNameType OriginalCollectionKey(CurrentCollectionName, CurrentShareType);
+	const FCollectionNameType NewCollectionKey(NewCollectionName, NewShareType);
 	{
-		const bool bUseSCC = ShouldUseSCC(NewShareType);
-		const FString NewCollectionFilename = GetCollectionFilename(NewCollectionName, NewShareType);
+		FCollectionLock_Write Guard(Lock);
 
-		// Create an exact copy of the collection using its new path - this will preserve its GUID and avoid losing hierarchy data
-		NewCollection = (*CollectionRefPtr)->Clone(NewCollectionFilename, bUseSCC, ECollectionCloneMode::Exact);
-		if (!AddCollection(NewCollection.ToSharedRef(), NewShareType))
+		TSharedRef<FCollection>* const CollectionRefPtr = AvailableCollections.Find(OriginalCollectionKey);
+		if (!CollectionRefPtr)
 		{
-			// Failed to add the collection, it already exists
-			LastError = LOCTEXT("Error_AlreadyExists", "The collection already exists.");
+			// The collection doesn't exist
+			if (OutError)
+			{
+				*OutError = LOCTEXT("Error_DoesntExist", "The collection doesn't exist.");
+			}
 			return false;
 		}
 
-		bool bForceCommitToRevisionControl = true;
-		if (!InternalSaveCollection(NewCollection.ToSharedRef(), LastError, bForceCommitToRevisionControl))
+		// Add the new collection
+		TSharedPtr<FCollection> NewCollection;
 		{
-			// Collection failed to save, remove it from the cache
-			RemoveCollection(NewCollection.ToSharedRef(), NewShareType);
-			return false;
-		}
-	}
+			const bool bUseSCC = ShouldUseSCC(NewShareType);
+			const FString NewCollectionFilename = GetCollectionFilename(NewCollectionName, NewShareType);
 
-	// Remove the old collection
-	{
-		if ((*CollectionRefPtr)->DeleteSourceFile(LastError))
-		{
-			CollectionFileCaches[CurrentShareType]->IgnoreDeletedFile((*CollectionRefPtr)->GetSourceFilename());
+			// Create an exact copy of the collection using its new path - this will preserve its GUID and avoid losing hierarchy data
+			NewCollection = (*CollectionRefPtr)->Clone(NewCollectionFilename, bUseSCC, ECollectionCloneMode::Exact);
+			if (!AddCollection(Guard, NewCollection.ToSharedRef(), NewShareType))
+			{
+				// Failed to add the collection, it already exists
+				if (OutError)
+				{
+					*OutError = LOCTEXT("Error_AlreadyExists", "The collection already exists.");
+				}
+				return false;
+			}
 
-			RemoveCollection(*CollectionRefPtr, CurrentShareType);
+			bool bForceCommitToRevisionControl = true;
+			if (!InternalSaveCollection(Guard, NewCollection.ToSharedRef(), OutError, bForceCommitToRevisionControl))
+			{
+				// Collection failed to save, remove it from the cache
+				RemoveCollection(Guard, NewCollection.ToSharedRef(), NewShareType);
+				return false;
+			}
 		}
-		else
+
+		// Remove the old collection
 		{
-			// Failed to remove the old collection, so remove the collection we created.
-			NewCollection->DeleteSourceFile(LastError);
-			RemoveCollection(NewCollection.ToSharedRef(), NewShareType);
-			return false;
+			FText UnusedError;
+			if ((*CollectionRefPtr)->DeleteSourceFile(OutError ? *OutError : UnusedError))
+			{
+				CollectionFileCaches[CurrentShareType]->IgnoreDeletedFile((*CollectionRefPtr)->GetSourceFilename());
+
+				RemoveCollection(Guard, *CollectionRefPtr, CurrentShareType);
+			}
+			else
+			{
+				// Failed to remove the old collection, so remove the collection we created.
+				NewCollection->DeleteSourceFile(OutError ? *OutError : UnusedError);
+				RemoveCollection(Guard, NewCollection.ToSharedRef(), NewShareType);
+				return false;
+			}
 		}
-	}
 
-	CollectionFileCaches[NewShareType]->IgnoreNewFile(NewCollection->GetSourceFilename());
+		CollectionFileCaches[NewShareType]->IgnoreNewFile(NewCollection->GetSourceFilename());
 
-	CollectionCache.HandleCollectionChanged();
+		CollectionCache->HandleCollectionChanged(Guard);
+	}
 
-	// Success
+	// Success, broadcast events outside of lock 
 	CollectionRenamedEvent.Broadcast(OriginalCollectionKey, NewCollectionKey);
 	return true;
 }
 
-bool FCollectionManager::ReparentCollection(FName CollectionName, ECollectionShareType::Type ShareType, FName ParentCollectionName, ECollectionShareType::Type ParentShareType)
+bool FCollectionManager::ReparentCollection(FName CollectionName, ECollectionShareType::Type ShareType, FName ParentCollectionName, ECollectionShareType::Type ParentShareType, FText* OutError) 
 {
 	if (!ensure(ShareType < ECollectionShareType::CST_All) || (!ParentCollectionName.IsNone() && !ensure(ParentShareType < ECollectionShareType::CST_All)))
 	{
 		// Bad share type
-		LastError = LOCTEXT("Error_Internal", "There was an internal error.");
+		if (OutError)
+		{
+			*OutError = LOCTEXT("Error_Internal", "There was an internal error.");
+		}
 		return false;
 	}
 
 	const FCollectionNameType CollectionKey(CollectionName, ShareType);
-	TSharedRef<FCollection>* const CollectionRefPtr = AvailableCollections.Find(CollectionKey);
-	if (!CollectionRefPtr)
-	{
-		// The collection doesn't exist
-		LastError = LOCTEXT("Error_DoesntExist", "The collection doesn't exist.");
-		return false;
-	}
-
-	const FGuid OldParentGuid = (*CollectionRefPtr)->GetParentCollectionGuid();
-	FGuid NewParentGuid;
-
 	TOptional<FCollectionNameType> OldParentCollectionKey;
 	TOptional<FCollectionNameType> NewParentCollectionKey;
-
-	if (!ParentCollectionName.IsNone())
 	{
-		// Find and set the new parent GUID
-		NewParentCollectionKey = FCollectionNameType(ParentCollectionName, ParentShareType);
-		TSharedRef<FCollection>* const ParentCollectionRefPtr = AvailableCollections.Find(NewParentCollectionKey.GetValue());
-		if (!ParentCollectionRefPtr)
+		FCollectionLock_Write Guard(Lock);
+		CollectionCache->UpdateCaches(Guard, ECollectionCacheFlags::RecursionWorker);
+
+		TSharedRef<FCollection>* const CollectionRefPtr = AvailableCollections.Find(CollectionKey);
+		if (!CollectionRefPtr)
 		{
 			// The collection doesn't exist
-			LastError = LOCTEXT("Error_DoesntExist", "The collection doesn't exist.");
+			if (OutError)
+			{
+				*OutError = LOCTEXT("Error_DoesntExist", "The collection doesn't exist.");
+			}
 			return false;
 		}
 
-		// Does the parent collection need saving in order to have a stable GUID?
-		if ((*ParentCollectionRefPtr)->GetCollectionVersion() < ECollectionVersion::AddedCollectionGuid)
+		const FGuid OldParentGuid = (*CollectionRefPtr)->GetParentCollectionGuid();
+		FGuid NewParentGuid;
+
+		if (!ParentCollectionName.IsNone())
 		{
-			bool bForceCommitToRevisionControl = false;
-			// Try and re-save the parent collection now
-			if (InternalSaveCollection(*ParentCollectionRefPtr, LastError, bForceCommitToRevisionControl))
+			// Find and set the new parent GUID
+			NewParentCollectionKey = FCollectionNameType(ParentCollectionName, ParentShareType);
+			TSharedRef<FCollection>* const ParentCollectionRefPtr = AvailableCollections.Find(NewParentCollectionKey.GetValue());
+			if (!ParentCollectionRefPtr)
 			{
-				CollectionFileCaches[ParentShareType]->IgnoreFileModification((*ParentCollectionRefPtr)->GetSourceFilename());
+				// The collection doesn't exist
+				if (OutError)
+				{
+					*OutError = LOCTEXT("Error_DoesntExist", "The collection doesn't exist.");
+				}
+				return false;
 			}
-			else
+
+			// Does the parent collection need saving in order to have a stable GUID?
+			if ((*ParentCollectionRefPtr)->GetCollectionVersion() < ECollectionVersion::AddedCollectionGuid)
+			{
+				bool bForceCommitToRevisionControl = false;
+				// Try and re-save the parent collection now
+				if (InternalSaveCollection(Guard, *ParentCollectionRefPtr, OutError, bForceCommitToRevisionControl))
+				{
+					CollectionFileCaches[ParentShareType]->IgnoreFileModification((*ParentCollectionRefPtr)->GetSourceFilename());
+				}
+				else
+				{
+					return false;
+				}
+			}
+
+			if (!IsValidParentCollection_Locked(Guard, CollectionName, ShareType, ParentCollectionName, ParentShareType, OutError))
 			{
 				return false;
 			}
+
+			NewParentGuid = (*ParentCollectionRefPtr)->GetCollectionGuid();
 		}
 
-		if (!IsValidParentCollection(CollectionName, ShareType, ParentCollectionName, ParentShareType))
+		// Anything changed?
+		if (OldParentGuid == NewParentGuid)
 		{
-			// IsValidParentCollection fills in LastError itself
-			return false;
+			return true;
 		}
 
-		NewParentGuid = (*ParentCollectionRefPtr)->GetCollectionGuid();
-	}
-
-	// Anything changed?
-	if (OldParentGuid == NewParentGuid)
-	{
-		return true;
-	}
-
-	(*CollectionRefPtr)->SetParentCollectionGuid(NewParentGuid);
-
-	// Try and save with the new parent GUID
-	bool bForceCommitToRevisionControl = false;
-	if (InternalSaveCollection(*CollectionRefPtr, LastError, bForceCommitToRevisionControl))
-	{
-		CollectionFileCaches[ShareType]->IgnoreFileModification((*CollectionRefPtr)->GetSourceFilename());
-	}
-	else
-	{
-		// Failed to save... rollback the collection to use its old parent GUID
-		(*CollectionRefPtr)->SetParentCollectionGuid(OldParentGuid);
-		return false;
-	}
+		(*CollectionRefPtr)->SetParentCollectionGuid(NewParentGuid);
 
-	CollectionCache.HandleCollectionChanged();
+		// Try and save with the new parent GUID
+		bool bForceCommitToRevisionControl = false;
+		if (InternalSaveCollection(Guard, *CollectionRefPtr, OutError, bForceCommitToRevisionControl))
+		{
+			CollectionFileCaches[ShareType]->IgnoreFileModification((*CollectionRefPtr)->GetSourceFilename());
+		}
+		else
+		{
+			// Failed to save... rollback the collection to use its old parent GUID
+			(*CollectionRefPtr)->SetParentCollectionGuid(OldParentGuid);
+			return false;
+		}
 
-	// Find the old parent so we can notify about the change
-	{
-		const FGuidToCollectionNamesMap& CachedCollectionNamesFromGuids = CollectionCache.GetCachedCollectionNamesFromGuids();
+		CollectionCache->HandleCollectionChanged(Guard);
+		CollectionCache->UpdateCaches(Guard, ECollectionCacheFlags::Names);
 
-		const FCollectionNameType* const OldParentCollectionKeyPtr = CachedCollectionNamesFromGuids.Find(OldParentGuid);
-		if (OldParentCollectionKeyPtr)
+		// Find the old parent so we can notify about the change
 		{
-			OldParentCollectionKey = *OldParentCollectionKeyPtr;
+			const TMap<FGuid, FCollectionNameType>& CachedCollectionNamesFromGuids = CollectionCache->GetCachedCollectionNamesFromGuids(Guard);
+
+			const FCollectionNameType* const OldParentCollectionKeyPtr = CachedCollectionNamesFromGuids.Find(OldParentGuid);
+			if (OldParentCollectionKeyPtr)
+			{
+				OldParentCollectionKey = *OldParentCollectionKeyPtr;
+			}
 		}
 	}
 
-	// Success
+	// Success, broadcast event outside of lock
 	CollectionReparentedEvent.Broadcast(CollectionKey, OldParentCollectionKey, NewParentCollectionKey);
 	return true;
 }
 
-bool FCollectionManager::DestroyCollection(FName CollectionName, ECollectionShareType::Type ShareType)
+bool FCollectionManager::DestroyCollection(FName CollectionName, ECollectionShareType::Type ShareType, FText* OutError)
 {
 	if (!ensure(ShareType < ECollectionShareType::CST_All))
 	{
 		// Bad share type
-		LastError = LOCTEXT("Error_Internal", "There was an internal error.");
+		if (OutError)
+		{
+			*OutError = LOCTEXT("Error_Internal", "There was an internal error.");
+		}
 		return false;
 	}
 
 	const FCollectionNameType CollectionKey(CollectionName, ShareType);
-	TSharedRef<FCollection>* const CollectionRefPtr = AvailableCollections.Find(CollectionKey);
-	if (!CollectionRefPtr)
 	{
-		// The collection doesn't exist
-		LastError = LOCTEXT("Error_DoesntExist", "The collection doesn't exist.");
-		return false;
-	}
+		FCollectionLock_Write Guard(Lock);
 
-	if ((*CollectionRefPtr)->DeleteSourceFile(LastError))
-	{
-		CollectionFileCaches[ShareType]->IgnoreDeletedFile((*CollectionRefPtr)->GetSourceFilename());
+		TSharedRef<FCollection>* const CollectionRefPtr = AvailableCollections.Find(CollectionKey);
+		if (!CollectionRefPtr)
+		{
+			// The collection doesn't exist
+			if (OutError)
+			{
+				*OutError = LOCTEXT("Error_DoesntExist", "The collection doesn't exist.");
+			}
+			return false;
+		}
 
-		RemoveCollection(*CollectionRefPtr, ShareType);
-		CollectionDestroyedEvent.Broadcast(CollectionKey);
-		return true;
-	}
-	else
-	{
-		// Failed to delete the source file
-		return false;
+		FText UnusedError;
+		if ((*CollectionRefPtr)->DeleteSourceFile(OutError ? *OutError : UnusedError))
+		{
+			CollectionFileCaches[ShareType]->IgnoreDeletedFile((*CollectionRefPtr)->GetSourceFilename());
+
+			RemoveCollection(Guard, *CollectionRefPtr, ShareType);
+		}
+		else
+		{
+			// Failed to delete the source file
+			return false;
+		}
 	}
+
+	// Broadcast event outside of lock
+	CollectionDestroyedEvent.Broadcast(CollectionKey);
+	return true;
 }
 
-bool FCollectionManager::AddToCollection(FName CollectionName, ECollectionShareType::Type ShareType, const FSoftObjectPath& ObjectPath)
+bool FCollectionManager::AddToCollection(FName CollectionName, ECollectionShareType::Type ShareType, const FSoftObjectPath& ObjectPath, FText* OutError)
 {
-	TArray<FSoftObjectPath> Paths;
-	Paths.Add(ObjectPath);
-	return AddToCollection(CollectionName, ShareType, Paths);
+	return AddToCollection(CollectionName, ShareType, MakeArrayView(&ObjectPath, 1));
 }
 
-bool FCollectionManager::AddToCollection(FName CollectionName, ECollectionShareType::Type ShareType, TConstArrayView<FSoftObjectPath> ObjectPaths, int32* OutNumAdded)
+bool FCollectionManager::AddToCollection(FName CollectionName, ECollectionShareType::Type ShareType, TConstArrayView<FSoftObjectPath> ObjectPaths, int32* OutNumAdded, FText* OutError)
 {
 	if (OutNumAdded)
 	{
@@ -1014,83 +1310,93 @@ bool FCollectionManager::AddToCollection(FName CollectionName, ECollectionShareT
 	if (!ensure(ShareType < ECollectionShareType::CST_All))
 	{
 		// Bad share type
-		LastError = LOCTEXT("Error_Internal", "There was an internal error.");
+		if (OutError)
+		{
+			*OutError = LOCTEXT("Error_Internal", "There was an internal error.");
+		}
 		return false;
 	}
 
 	const FCollectionNameType CollectionKey(CollectionName, ShareType);
-	TSharedRef<FCollection>* const CollectionRefPtr = AvailableCollections.Find(CollectionKey);
-	if (!CollectionRefPtr)
-	{
-		// Collection doesn't exist
-		LastError = LOCTEXT("Error_DoesntExist", "The collection doesn't exist.");
-		return false;
-	}
-
-	if ((*CollectionRefPtr)->GetStorageMode() != ECollectionStorageMode::Static)
-	{
-		LastError = LOCTEXT("Error_AddNeedsStaticCollection", "Objects can only be added to static collections.");
-		return false;
-	}
-	
-	int32 NumAdded = 0;
-	for (const FSoftObjectPath& ObjectPath : ObjectPaths)
-	{
-		if ((*CollectionRefPtr)->AddObjectToCollection(ObjectPath))
+	{ 
+		FCollectionLock_Write Guard(Lock);
+		TSharedRef<FCollection>* const CollectionRefPtr = AvailableCollections.Find(CollectionKey);
+		if (!CollectionRefPtr)
 		{
-			NumAdded++;
+			// Collection doesn't exist
+			if (OutError)
+			{
+				*OutError = LOCTEXT("Error_DoesntExist", "The collection doesn't exist.");
+			}
+			return false;
 		}
-	}
 
-	if (NumAdded > 0)
-	{
-		constexpr bool bForceCommitToRevisionControl = false;
-		if (InternalSaveCollection(*CollectionRefPtr, LastError, bForceCommitToRevisionControl))
+		if ((*CollectionRefPtr)->GetStorageMode() != ECollectionStorageMode::Static)
 		{
-			CollectionFileCaches[ShareType]->IgnoreFileModification((*CollectionRefPtr)->GetSourceFilename());
-
-			// Added and saved
-			if (OutNumAdded)
+			if (OutError)
 			{
-				*OutNumAdded = NumAdded;
+				*OutError = LOCTEXT("Error_AddNeedsStaticCollection", "Objects can only be added to static collections.");
 			}
+			return false;
+		}
+		
+		int32 NumAdded = 0;
+		for (const FSoftObjectPath& ObjectPath : ObjectPaths)
+		{
+			if ((*CollectionRefPtr)->AddObjectToCollection(ObjectPath))
+			{
+				NumAdded++;
+			}
+		}
 
-			CollectionCache.HandleCollectionChanged();
-			AssetsAddedToCollectionDelegate.Broadcast(CollectionKey, ObjectPaths);
-PRAGMA_DISABLE_DEPRECATION_WARNINGS
-			if (AssetsAddedEvent.IsBound()) // Avoid conversion if nothing is bound 
+		if (NumAdded > 0)
+		{
+			constexpr bool bForceCommitToRevisionControl = false;
+			if (InternalSaveCollection(Guard, *CollectionRefPtr, OutError, bForceCommitToRevisionControl))
 			{
-				AssetsAddedEvent.Broadcast(CollectionKey, UE::SoftObjectPath::Private::ConvertSoftObjectPaths(ObjectPaths));
+				CollectionFileCaches[ShareType]->IgnoreFileModification((*CollectionRefPtr)->GetSourceFilename());
+
+				// Added and saved
+				if (OutNumAdded)
+				{
+					*OutNumAdded = NumAdded;
+				}
+
+				CollectionCache->HandleCollectionChanged(Guard);
+				// Fall out of scope to return 
+			}
+			else
+			{
+				// Added but not saved, revert the add
+				for (const FSoftObjectPath& ObjectPath : ObjectPaths)
+				{
+					(*CollectionRefPtr)->RemoveObjectFromCollection(ObjectPath);
+				}
+				return false;
 			}
-PRAGMA_ENABLE_DEPRECATION_WARNINGS
-			return true;
 		}
 		else
 		{
-			// Added but not saved, revert the add
-			for (const FSoftObjectPath& ObjectPath : ObjectPaths)
+			// Failed to add, all of the objects were already in the collection
+			if (OutError)
 			{
-				(*CollectionRefPtr)->RemoveObjectFromCollection(ObjectPath);
+				*OutError = LOCTEXT("Error_AlreadyInCollection", "All of the assets were already in the collection.");
 			}
 			return false;
 		}
 	}
-	else
-	{
-		// Failed to add, all of the objects were already in the collection
-		LastError = LOCTEXT("Error_AlreadyInCollection", "All of the assets were already in the collection.");
-		return false;
-	}
+	   
+	// Broadast event out of lock
+	AssetsAddedToCollectionDelegate.Broadcast(CollectionKey, ObjectPaths);
+	return true;
 }
 
-bool FCollectionManager::RemoveFromCollection(FName CollectionName, ECollectionShareType::Type ShareType, const FSoftObjectPath& ObjectPath)
+bool FCollectionManager::RemoveFromCollection(FName CollectionName, ECollectionShareType::Type ShareType, const FSoftObjectPath& ObjectPath, FText* OutError)
 {
-	TArray<FSoftObjectPath> Paths;
-	Paths.Add(ObjectPath);
-	return RemoveFromCollection(CollectionName, ShareType, Paths);
+	return RemoveFromCollection(CollectionName, ShareType, MakeArrayView(&ObjectPath, 1), nullptr, OutError);
 }
 
-bool FCollectionManager::RemoveFromCollection(FName CollectionName, ECollectionShareType::Type ShareType, TConstArrayView<FSoftObjectPath> ObjectPaths, int32* OutNumRemoved)
+bool FCollectionManager::RemoveFromCollection(FName CollectionName, ECollectionShareType::Type ShareType, TConstArrayView<FSoftObjectPath> ObjectPaths, int32* OutNumRemoved, FText* OutError)
 {
 	if (OutNumRemoved)
 	{
@@ -1100,44 +1406,66 @@ bool FCollectionManager::RemoveFromCollection(FName CollectionName, ECollectionS
 	if (!ensure(ShareType < ECollectionShareType::CST_All))
 	{
 		// Bad share type
-		LastError = LOCTEXT("Error_Internal", "There was an internal error.");
+		if (OutError)
+		{
+			*OutError = LOCTEXT("Error_Internal", "There was an internal error.");
+		}
 		return false;
 	}
 
 	const FCollectionNameType CollectionKey(CollectionName, ShareType);
-	TSharedRef<FCollection>* const CollectionRefPtr = AvailableCollections.Find(CollectionKey);
-	if (!CollectionRefPtr)
-	{
-		// Collection not found
-		LastError = LOCTEXT("Error_DoesntExist", "The collection doesn't exist.");
-		return false;
-	}
+	{ 
+		FCollectionLock_Write Guard(Lock);
+		TSharedRef<FCollection>* const CollectionRefPtr = AvailableCollections.Find(CollectionKey);
+		if (!CollectionRefPtr)
+		{
+			// Collection not found
+			if (OutError)
+			{
+				*OutError = LOCTEXT("Error_DoesntExist", "The collection doesn't exist.");
+			}
+			return false;
+		}
 
-	if ((*CollectionRefPtr)->GetStorageMode() != ECollectionStorageMode::Static)
-	{
-		LastError = LOCTEXT("Error_RemoveNeedsStaticCollection", "Objects can only be removed from static collections.");
-		return false;
-	}
-	
-	TArray<FSoftObjectPath> RemovedAssets;
-	for (const FSoftObjectPath& ObjectPath : ObjectPaths)
-	{
-		if ((*CollectionRefPtr)->RemoveObjectFromCollection(ObjectPath))
+		if ((*CollectionRefPtr)->GetStorageMode() != ECollectionStorageMode::Static)
 		{
-			RemovedAssets.Add(ObjectPath);
+			if (OutError)
+			{
+				*OutError = LOCTEXT("Error_RemoveNeedsStaticCollection", "Objects can only be removed from static collections.");
+			}
+			return false;
+		}
+		
+		TArray<FSoftObjectPath> RemovedAssets;
+		for (const FSoftObjectPath& ObjectPath : ObjectPaths)
+		{
+			if ((*CollectionRefPtr)->RemoveObjectFromCollection(ObjectPath))
+			{
+				RemovedAssets.Add(ObjectPath);
+			}
+		}
+
+		if (RemovedAssets.Num() == 0)
+		{
+			// Failed to remove, none of the objects were in the collection
+			if (OutError)
+			{
+				*OutError = LOCTEXT("Error_NotInCollection", "None of the assets were in the collection.");
+			}
+			return false;
+		}
+
+		constexpr bool bForceCommitToRevisionControl = false;
+		if (!InternalSaveCollection(Guard, *CollectionRefPtr, OutError, bForceCommitToRevisionControl))
+		{
+			// Removed but not saved, revert the remove
+			for (const FSoftObjectPath& RemovedAssetName : RemovedAssets)
+			{
+				(*CollectionRefPtr)->AddObjectToCollection(RemovedAssetName);
+			}
+			return false;
 		}
-	}
 
-	if (RemovedAssets.Num() == 0)
-	{
-		// Failed to remove, none of the objects were in the collection
-		LastError = LOCTEXT("Error_NotInCollection", "None of the assets were in the collection.");
-		return false;
-	}
-			
-	constexpr bool bForceCommitToRevisionControl = false;
-	if (InternalSaveCollection(*CollectionRefPtr, LastError, bForceCommitToRevisionControl))
-	{
 		CollectionFileCaches[ShareType]->IgnoreFileModification((*CollectionRefPtr)->GetSourceFilename());
 
 		// Removed and saved
@@ -1146,87 +1474,97 @@ bool FCollectionManager::RemoveFromCollection(FName CollectionName, ECollectionS
 			*OutNumRemoved = RemovedAssets.Num();
 		}
 
-		CollectionCache.HandleCollectionChanged();
-		AssetsRemovedFromCollectionDelegate.Broadcast(CollectionKey, ObjectPaths);
-PRAGMA_DISABLE_DEPRECATION_WARNINGS
-		if (AssetsRemovedEvent.IsBound())
-		{
-			AssetsRemovedEvent.Broadcast(CollectionKey, UE::SoftObjectPath::Private::ConvertSoftObjectPaths(ObjectPaths));
-		}
-PRAGMA_ENABLE_DEPRECATION_WARNINGS
-		return true;
-	}
-	else
-	{
-		// Removed but not saved, revert the remove
-		for (const FSoftObjectPath& RemovedAssetName : RemovedAssets)
-		{
-			(*CollectionRefPtr)->AddObjectToCollection(RemovedAssetName);
-		}
-		return false;
+		CollectionCache->HandleCollectionChanged(Guard);
 	}
+
+	// Broadcast event out of lock
+	AssetsRemovedFromCollectionDelegate.Broadcast(CollectionKey, ObjectPaths);
+	return true;
 }
 
-bool FCollectionManager::SetDynamicQueryText(FName CollectionName, ECollectionShareType::Type ShareType, const FString& InQueryText)
+bool FCollectionManager::SetDynamicQueryText(FName CollectionName, ECollectionShareType::Type ShareType, const FString& InQueryText, FText* OutError)
 {
 	if (!ensure(ShareType < ECollectionShareType::CST_All))
 	{
 		// Bad share type
-		LastError = LOCTEXT("Error_Internal", "There was an internal error.");
+		if (OutError)
+		{
+			*OutError = LOCTEXT("Error_Internal", "There was an internal error.");
+		}
 		return false;
 	}
 
 	const FCollectionNameType CollectionKey(CollectionName, ShareType);
-	TSharedRef<FCollection>* const CollectionRefPtr = AvailableCollections.Find(CollectionKey);
-	if (!CollectionRefPtr)
-	{
-		// Collection doesn't exist
-		LastError = LOCTEXT("Error_DoesntExist", "The collection doesn't exist.");
-		return false;
-	}
+	{ 
+		FCollectionLock_Write Guard(Lock);
+		TSharedRef<FCollection>* const CollectionRefPtr = AvailableCollections.Find(CollectionKey);
+		if (!CollectionRefPtr)
+		{
+			// Collection doesn't exist
+			if (OutError)
+			{
+				*OutError = LOCTEXT("Error_DoesntExist", "The collection doesn't exist.");
+			}
+			return false;
+		}
 
-	if ((*CollectionRefPtr)->GetStorageMode() != ECollectionStorageMode::Dynamic)
-	{
-		LastError = LOCTEXT("Error_SetNeedsDynamicCollection", "Search queries can only be set on dynamic collections.");
-		return false;
-	}
+		if ((*CollectionRefPtr)->GetStorageMode() != ECollectionStorageMode::Dynamic)
+		{
+			if (OutError)
+			{
+				*OutError = LOCTEXT("Error_SetNeedsDynamicCollection", "Search queries can only be set on dynamic collections.");
+			}
+			return false;
+		}
 
-	(*CollectionRefPtr)->SetDynamicQueryText(InQueryText);
-	
-	constexpr bool bForceCommitToRevisionControl = true;
-	if (InternalSaveCollection(*CollectionRefPtr, LastError, bForceCommitToRevisionControl))
-	{
+		(*CollectionRefPtr)->SetDynamicQueryText(InQueryText);
+		
+		constexpr bool bForceCommitToRevisionControl = true;
+		if (!InternalSaveCollection(Guard, *CollectionRefPtr, OutError, bForceCommitToRevisionControl))
+		{
+			return false;
+		}
 		CollectionFileCaches[ShareType]->IgnoreFileModification((*CollectionRefPtr)->GetSourceFilename());
-
-		CollectionCache.HandleCollectionChanged();
-		CollectionUpdatedEvent.Broadcast(CollectionKey);
-		return true;
+		CollectionCache->HandleCollectionChanged(Guard);
 	}
 
-	return false;
+	// Broadcast event outside of lock
+	CollectionUpdatedEvent.Broadcast(CollectionKey);
+	return true;
+
 }
 
-bool FCollectionManager::GetDynamicQueryText(FName CollectionName, ECollectionShareType::Type ShareType, FString& OutQueryText) const
+bool FCollectionManager::GetDynamicQueryText(FName CollectionName, ECollectionShareType::Type ShareType, FString& OutQueryText, FText* OutError) const
 {
 	if (!ensure(ShareType < ECollectionShareType::CST_All))
 	{
 		// Bad share type
-		LastError = LOCTEXT("Error_Internal", "There was an internal error.");
+		if (OutError)
+		{
+			*OutError = LOCTEXT("Error_Internal", "There was an internal error.");
+		}
 		return false;
 	}
 
+	FCollectionLock_Read Guard(Lock);
 	const FCollectionNameType CollectionKey(CollectionName, ShareType);
 	const TSharedRef<FCollection>* const CollectionRefPtr = AvailableCollections.Find(CollectionKey);
 	if (!CollectionRefPtr)
 	{
 		// Collection doesn't exist
-		LastError = LOCTEXT("Error_DoesntExist", "The collection doesn't exist.");
+		if (OutError)
+		{
+			*OutError = LOCTEXT("Error_DoesntExist", "The collection doesn't exist.");
+		}
 		return false;
 	}
 
 	if ((*CollectionRefPtr)->GetStorageMode() != ECollectionStorageMode::Dynamic)
 	{
-		LastError = LOCTEXT("Error_GetNeedsDynamicCollection", "Search queries can only be got from dynamic collections.");
+		if (OutError)
+		{
+			*OutError = LOCTEXT("Error_GetNeedsDynamicCollection", "Search queries can only be got from dynamic collections.");
+		}
 		return false;
 	}
 
@@ -1234,27 +1572,37 @@ bool FCollectionManager::GetDynamicQueryText(FName CollectionName, ECollectionSh
 	return true;
 }
 
-bool FCollectionManager::TestDynamicQuery(FName CollectionName, ECollectionShareType::Type ShareType, const ITextFilterExpressionContext& InContext, bool& OutResult) const
+bool FCollectionManager::TestDynamicQuery(FName CollectionName, ECollectionShareType::Type ShareType, const ITextFilterExpressionContext& InContext, bool& OutResult, FText* OutError) const
 {
 	if (!ensure(ShareType < ECollectionShareType::CST_All))
 	{
 		// Bad share type
-		LastError = LOCTEXT("Error_Internal", "There was an internal error.");
+		if (OutError)
+		{
+			*OutError = LOCTEXT("Error_Internal", "There was an internal error.");
+		}
 		return false;
 	}
 
+	FCollectionLock_Read Guard(Lock);
 	const FCollectionNameType CollectionKey(CollectionName, ShareType);
 	const TSharedRef<FCollection>* const CollectionRefPtr = AvailableCollections.Find(CollectionKey);
 	if (!CollectionRefPtr)
 	{
 		// Collection doesn't exist
-		LastError = LOCTEXT("Error_DoesntExist", "The collection doesn't exist.");
+		if (OutError)
+		{
+			*OutError = LOCTEXT("Error_DoesntExist", "The collection doesn't exist.");
+		}
 		return false;
 	}
 
 	if ((*CollectionRefPtr)->GetStorageMode() != ECollectionStorageMode::Dynamic)
 	{
-		LastError = LOCTEXT("Error_TestNeedsDynamicCollection", "Search queries can only be tested on dynamic collections.");
+		if (OutError)
+		{
+			*OutError = LOCTEXT("Error_TestNeedsDynamicCollection", "Search queries can only be tested on dynamic collections.");
+		}
 		return false;
 	}
 
@@ -1262,58 +1610,80 @@ bool FCollectionManager::TestDynamicQuery(FName CollectionName, ECollectionShare
 	return true;
 }
 
-bool FCollectionManager::EmptyCollection(FName CollectionName, ECollectionShareType::Type ShareType)
+bool FCollectionManager::EmptyCollection(FName CollectionName, ECollectionShareType::Type ShareType, FText* OutError)
 {
 	if (!ensure(ShareType < ECollectionShareType::CST_All))
 	{
 		// Bad share type
-		LastError = LOCTEXT("Error_Internal", "There was an internal error.");
+		if (OutError)
+		{
+			*OutError = LOCTEXT("Error_Internal", "There was an internal error.");
+		}
 		return false;
 	}
 
 	const FCollectionNameType CollectionKey(CollectionName, ShareType);
-	const TSharedRef<FCollection>* const CollectionRefPtr = AvailableCollections.Find(CollectionKey);
-	if (!CollectionRefPtr)
-	{
-		// Collection doesn't exist
-		LastError = LOCTEXT("Error_DoesntExist", "The collection doesn't exist.");
-		return false;
-	}
+	{ 
+		FCollectionLock_Write Guard(Lock);
+		const TSharedRef<FCollection>* const CollectionRefPtr = AvailableCollections.Find(CollectionKey);
+		if (!CollectionRefPtr)
+		{
+			// Collection doesn't exist
+			if (OutError)
+			{
+				*OutError = LOCTEXT("Error_DoesntExist", "The collection doesn't exist.");
+			}
+			return false;
+		}
 
-	if ((*CollectionRefPtr)->IsEmpty())
-	{
-		// Already empty - nothing to do
-		return true;
-	}
+		if ((*CollectionRefPtr)->IsEmpty())
+		{
+			// Already empty - nothing to do
+			return true;
+		}
 
-	(*CollectionRefPtr)->Empty();
-	
-	constexpr bool bForceCommitToRevisionControl = true;
-	if (InternalSaveCollection(*CollectionRefPtr, LastError, bForceCommitToRevisionControl))
-	{
+		(*CollectionRefPtr)->Empty();
+		
+		constexpr bool bForceCommitToRevisionControl = true;
+		if (InternalSaveCollection(Guard, *CollectionRefPtr, OutError, bForceCommitToRevisionControl))
+		{
+			return false;
+		}
 		CollectionFileCaches[ShareType]->IgnoreFileModification((*CollectionRefPtr)->GetSourceFilename());
 
-		CollectionCache.HandleCollectionChanged();
-		CollectionUpdatedEvent.Broadcast(CollectionKey);
-		return true;
+		CollectionCache->HandleCollectionChanged(Guard);
 	}
 
-	return false;
+	// Broadcast event outside of lock
+	CollectionUpdatedEvent.Broadcast(CollectionKey);
+	return true;
 }
 
-bool FCollectionManager::SaveCollection(FName CollectionName, ECollectionShareType::Type ShareType)
+bool FCollectionManager::SaveCollection(FName CollectionName, ECollectionShareType::Type ShareType, FText* OutError)
 {
 	if (!ensure(ShareType < ECollectionShareType::CST_All))
 	{
 		// Bad share type
-		LastError = LOCTEXT("Error_Internal", "There was an internal error.");
+		if (OutError)
+		{
+			*OutError = LOCTEXT("Error_Internal", "There was an internal error.");
+		}
 		return false;
 	}
 
 	const FCollectionNameType CollectionKey(CollectionName, ShareType);
-	const TSharedRef<FCollection>* const CollectionRefPtr = AvailableCollections.Find(CollectionKey);
-	if (CollectionRefPtr)
-	{
+	{ 
+		FCollectionLock_Write Guard(Lock);
+		const TSharedRef<FCollection>* const CollectionRefPtr = AvailableCollections.Find(CollectionKey);
+		if (!CollectionRefPtr)
+		{
+			if (OutError)
+			{
+				*OutError = LOCTEXT("Error_DoesntExist", "The collection doesn't exist.");
+			}
+			return false;
+		}
+
 		FCollectionStatusInfo StatusInfo = (*CollectionRefPtr)->GetStatusInfo();
 
 		const bool bNeedsSave = StatusInfo.bIsDirty || (StatusInfo.SCCState.IsValid() && StatusInfo.SCCState->IsModified());
@@ -1324,64 +1694,76 @@ bool FCollectionManager::SaveCollection(FName CollectionName, ECollectionShareTy
 		}
 
 		constexpr bool bForceCommitToRevisionControl = true;
-		if (InternalSaveCollection(*CollectionRefPtr, LastError, bForceCommitToRevisionControl))
+		if (!InternalSaveCollection(Guard, *CollectionRefPtr, OutError, bForceCommitToRevisionControl))
 		{
-			CollectionFileCaches[ShareType]->IgnoreFileModification((*CollectionRefPtr)->GetSourceFilename());
-
-			CollectionCache.HandleCollectionChanged();
-			CollectionUpdatedEvent.Broadcast(CollectionKey);
-			return true;
+			return false;
 		}
-	}
-	else
-	{
-		LastError = LOCTEXT("Error_DoesntExist", "The collection doesn't exist.");
-	}
 
-	return false;
+		CollectionFileCaches[ShareType]->IgnoreFileModification((*CollectionRefPtr)->GetSourceFilename());
+
+		CollectionCache->HandleCollectionChanged(Guard);
+	}
+	
+	// Broadcast event out of lock
+	CollectionUpdatedEvent.Broadcast(CollectionKey);
+	return true;
 }
 
-bool FCollectionManager::UpdateCollection(FName CollectionName, ECollectionShareType::Type ShareType)
+bool FCollectionManager::UpdateCollection(FName CollectionName, ECollectionShareType::Type ShareType, FText* OutError)
 {
 	if (!ensure(ShareType < ECollectionShareType::CST_All))
 	{
 		// Bad share type
-		LastError = LOCTEXT("Error_Internal", "There was an internal error.");
+		if (OutError)
+		{
+			*OutError = LOCTEXT("Error_Internal", "There was an internal error.");
+		}
 		return false;
 	}
 
 	const FCollectionNameType CollectionKey(CollectionName, ShareType);
-	const TSharedRef<FCollection>* const CollectionRefPtr = AvailableCollections.Find(CollectionKey);
-	if (CollectionRefPtr)
-	{
-		if ((*CollectionRefPtr)->Update(LastError))
+	{ 
+		FCollectionLock_Write Guard(Lock);
+		const TSharedRef<FCollection>* const CollectionRefPtr = AvailableCollections.Find(CollectionKey);
+		if (!CollectionRefPtr)
 		{
-			CollectionFileCaches[ShareType]->IgnoreFileModification((*CollectionRefPtr)->GetSourceFilename());
+			if (OutError)
+			{
+				*OutError = LOCTEXT("Error_DoesntExist", "The collection doesn't exist.");
+			}
+			return false;
+		}
 
-			CollectionCache.HandleCollectionChanged();
-			CollectionUpdatedEvent.Broadcast(CollectionKey);
-			return true;
+		FText UnusedError;
+		if (!(*CollectionRefPtr)->Update(OutError ? *OutError : UnusedError))
+		{
+			return false;
 		}
-	}
-	else
-	{
-		LastError = LOCTEXT("Error_DoesntExist", "The collection doesn't exist.");
+
+		CollectionFileCaches[ShareType]->IgnoreFileModification((*CollectionRefPtr)->GetSourceFilename());
+		CollectionCache->HandleCollectionChanged(Guard);
 	}
 
-	return false;
+	// Broadcast event outside of lock
+	CollectionUpdatedEvent.Broadcast(CollectionKey);
+	return true;
 }
 
-bool FCollectionManager::GetCollectionStatusInfo(FName CollectionName, ECollectionShareType::Type ShareType, FCollectionStatusInfo& OutStatusInfo) const
+bool FCollectionManager::GetCollectionStatusInfo(FName CollectionName, ECollectionShareType::Type ShareType, FCollectionStatusInfo& OutStatusInfo, FText* OutError) const
 {
 	TRACE_CPUPROFILER_EVENT_SCOPE(FCollectionManager::GetCollectionStatusInfo);
 
 	if (!ensure(ShareType < ECollectionShareType::CST_All))
 	{
 		// Bad share type
-		LastError = LOCTEXT("Error_Internal", "There was an internal error.");
+		if (OutError)
+		{
+			*OutError = LOCTEXT("Error_Internal", "There was an internal error.");
+		}
 		return true;
 	}
 
+	FCollectionLock_Read Guard(Lock);
 	const FCollectionNameType CollectionKey(CollectionName, ShareType);
 	const TSharedRef<FCollection>* const CollectionRefPtr = AvailableCollections.Find(CollectionKey);
 	if (CollectionRefPtr)
@@ -1391,7 +1773,10 @@ bool FCollectionManager::GetCollectionStatusInfo(FName CollectionName, ECollecti
 	}
 	else
 	{
-		LastError = LOCTEXT("Error_DoesntExist", "The collection doesn't exist.");
+		if (OutError)
+		{
+			*OutError = LOCTEXT("Error_DoesntExist", "The collection doesn't exist.");
+		}
 	}
 
 	return false;
@@ -1399,7 +1784,9 @@ bool FCollectionManager::GetCollectionStatusInfo(FName CollectionName, ECollecti
 
 bool FCollectionManager::HasCollectionColors(TArray<FLinearColor>* OutColors) const
 {
-	const FCollectionColorArray& CollectionColors = CollectionCache.GetCachedColors();
+	FCollectionLock_RW Guard(Lock);
+	CollectionCache->UpdateCaches(Guard, ECollectionCacheFlags::Colors);
+	const TArray<FLinearColor>& CollectionColors = CollectionCache->GetCachedColors(Guard);
 	if (OutColors)
 	{
 		*OutColors = CollectionColors;
@@ -1407,15 +1794,19 @@ bool FCollectionManager::HasCollectionColors(TArray<FLinearColor>* OutColors) co
 	return CollectionColors.Num() > 0;
 }
 
-bool FCollectionManager::GetCollectionColor(FName CollectionName, ECollectionShareType::Type ShareType, TOptional<FLinearColor>& OutColor) const
+bool FCollectionManager::GetCollectionColor(FName CollectionName, ECollectionShareType::Type ShareType, TOptional<FLinearColor>& OutColor, FText* OutError) const
 {
 	if (!ensure(ShareType < ECollectionShareType::CST_All))
 	{
 		// Bad share type
-		LastError = LOCTEXT("Error_Internal", "There was an internal error.");
+		if (OutError)
+		{
+			*OutError = LOCTEXT("Error_Internal", "There was an internal error.");
+		}
 		return true;
 	}
 
+	FCollectionLock_Read Guard(Lock);
 	const FCollectionNameType CollectionKey(CollectionName, ShareType);
 	const TSharedRef<FCollection>* const CollectionRefPtr = AvailableCollections.Find(CollectionKey);
 	if (CollectionRefPtr)
@@ -1425,54 +1816,70 @@ bool FCollectionManager::GetCollectionColor(FName CollectionName, ECollectionSha
 	}
 	else
 	{
-		LastError = LOCTEXT("Error_DoesntExist", "The collection doesn't exist.");
+		if (OutError)
+		{
+			*OutError = LOCTEXT("Error_DoesntExist", "The collection doesn't exist.");
+		}
 	}
 
 	return false;
 }
 
-bool FCollectionManager::SetCollectionColor(FName CollectionName, ECollectionShareType::Type ShareType, const TOptional<FLinearColor>& NewColor)
+bool FCollectionManager::SetCollectionColor(FName CollectionName, ECollectionShareType::Type ShareType, const TOptional<FLinearColor>& NewColor, FText* OutError)
 {
 	if (!ensure(ShareType < ECollectionShareType::CST_All))
 	{
 		// Bad share type
-		LastError = LOCTEXT("Error_Internal", "There was an internal error.");
+		if (OutError)
+		{
+			*OutError = LOCTEXT("Error_Internal", "There was an internal error.");
+		}
 		return true;
 	}
 
 	const FCollectionNameType CollectionKey(CollectionName, ShareType);
-	const TSharedRef<FCollection>* const CollectionRefPtr = AvailableCollections.Find(CollectionKey);
-	if (CollectionRefPtr)
-	{
+	{ 
+		FCollectionLock_Write Guard(Lock);
+		const TSharedRef<FCollection>* const CollectionRefPtr = AvailableCollections.Find(CollectionKey);
+		if (!CollectionRefPtr)
+		{
+			if (OutError)
+			{
+				*OutError = LOCTEXT("Error_DoesntExist", "The collection doesn't exist.");
+			}
+		}
+
 		(*CollectionRefPtr)->SetCollectionColor(NewColor);
 
 		constexpr bool bForceCommitToRevisionControl = false;
-		if (InternalSaveCollection(*CollectionRefPtr, LastError, bForceCommitToRevisionControl))
+		if (!InternalSaveCollection(Guard, *CollectionRefPtr, OutError, bForceCommitToRevisionControl))
 		{
-			CollectionFileCaches[ShareType]->IgnoreFileModification((*CollectionRefPtr)->GetSourceFilename());
-			
-			CollectionCache.HandleCollectionChanged();
-			CollectionUpdatedEvent.Broadcast(CollectionKey);
-			return true;
+			return false;
 		}
+		
+		CollectionFileCaches[ShareType]->IgnoreFileModification((*CollectionRefPtr)->GetSourceFilename());
+		
+		CollectionCache->HandleCollectionChanged(Guard);
 	}
-	else
-	{
-		LastError = LOCTEXT("Error_DoesntExist", "The collection doesn't exist.");
-	}
-
-	return false;
+	
+	// Broadcast event outside of lock 
+	CollectionUpdatedEvent.Broadcast(CollectionKey);
+	return true;
 }
 
-bool FCollectionManager::GetCollectionStorageMode(FName CollectionName, ECollectionShareType::Type ShareType, ECollectionStorageMode::Type& OutStorageMode) const
+bool FCollectionManager::GetCollectionStorageMode(FName CollectionName, ECollectionShareType::Type ShareType, ECollectionStorageMode::Type& OutStorageMode, FText* OutError) const
 {
 	if (!ensure(ShareType < ECollectionShareType::CST_All))
 	{
 		// Bad share type
-		LastError = LOCTEXT("Error_Internal", "There was an internal error.");
+		if (OutError)
+		{
+			*OutError = LOCTEXT("Error_Internal", "There was an internal error.");
+		}
 		return true;
 	}
 
+	FCollectionLock_Read Guard(Lock);
 	const FCollectionNameType CollectionKey(CollectionName, ShareType);
 	const TSharedRef<FCollection>* const CollectionRefPtr = AvailableCollections.Find(CollectionKey);
 	if (CollectionRefPtr)
@@ -1482,26 +1889,34 @@ bool FCollectionManager::GetCollectionStorageMode(FName CollectionName, ECollect
 	}
 	else
 	{
-		LastError = LOCTEXT("Error_DoesntExist", "The collection doesn't exist.");
+		if (OutError)
+		{
+			*OutError = LOCTEXT("Error_DoesntExist", "The collection doesn't exist.");
+		}
 	}
 
 	return false;
 }
 
-bool FCollectionManager::IsObjectInCollection(const FSoftObjectPath& ObjectPath, FName CollectionName, ECollectionShareType::Type ShareType, ECollectionRecursionFlags::Flags RecursionMode) const
+bool FCollectionManager::IsObjectInCollection(const FSoftObjectPath& ObjectPath, FName CollectionName, ECollectionShareType::Type ShareType,ECollectionRecursionFlags::Flags RecursionMode, FText* OutError) const
 {
 	if (!ensure(ShareType < ECollectionShareType::CST_All))
 	{
 		// Bad share type
-		LastError = LOCTEXT("Error_Internal", "There was an internal error.");
+		if (OutError)
+		{
+			*OutError = LOCTEXT("Error_Internal", "There was an internal error.");
+		}
 		return true;
 	}
 
+	FCollectionLock_RW Guard(Lock);
+	CollectionCache->UpdateCaches(Guard, ECollectionCacheFlags::RecursionWorker);
 	bool bFoundObject = false;
 
-	auto IsObjectInCollectionWorker = [&](const FCollectionNameType& InCollectionKey, ECollectionRecursionFlags::Flag InReason) -> FCollectionManagerCache::ERecursiveWorkerFlowControl
+	auto IsObjectInCollectionWorker = [AvailableCollections=&AvailableCollections, ObjectPath, &bFoundObject](const FCollectionNameType& InCollectionKey, ECollectionRecursionFlags::Flag InReason) -> FCollectionManagerCache::ERecursiveWorkerFlowControl
 	{
-		const TSharedRef<FCollection>* const CollectionRefPtr = AvailableCollections.Find(InCollectionKey);
+		const TSharedRef<FCollection>* const CollectionRefPtr = AvailableCollections->Find(InCollectionKey);
 		if (CollectionRefPtr)
 		{
 			bFoundObject = (*CollectionRefPtr)->IsObjectInCollection(ObjectPath);
@@ -1509,17 +1924,27 @@ bool FCollectionManager::IsObjectInCollection(const FSoftObjectPath& ObjectPath,
 		return (bFoundObject) ? FCollectionManagerCache::ERecursiveWorkerFlowControl::Stop : FCollectionManagerCache::ERecursiveWorkerFlowControl::Continue;
 	};
 
-	CollectionCache.RecursionHelper_DoWork(FCollectionNameType(CollectionName, ShareType), RecursionMode, IsObjectInCollectionWorker);
+	CollectionCache->RecursionHelper_DoWork(Guard, FCollectionNameType(CollectionName, ShareType), RecursionMode, IsObjectInCollectionWorker);
 
 	return bFoundObject;
 }
 
-bool FCollectionManager::IsValidParentCollection(FName CollectionName, ECollectionShareType::Type ShareType, FName ParentCollectionName, ECollectionShareType::Type ParentShareType) const
+bool FCollectionManager::IsValidParentCollection(FName CollectionName, ECollectionShareType::Type ShareType, FName ParentCollectionName,ECollectionShareType::Type ParentShareType, FText* OutError) const
+{
+	FCollectionLock_RW Guard(Lock);
+	CollectionCache->UpdateCaches(Guard, ECollectionCacheFlags::RecursionWorker);
+	return IsValidParentCollection_Locked(Guard, CollectionName, ShareType, ParentCollectionName, ParentShareType, OutError);
+}
+
+bool FCollectionManager::IsValidParentCollection_Locked(FCollectionLock& InGuard, FName CollectionName, ECollectionShareType::Type ShareType, FName ParentCollectionName, ECollectionShareType::Type ParentShareType, FText* OutError) const 
 {
 	if (!ensure(ShareType < ECollectionShareType::CST_All) || (!ParentCollectionName.IsNone() && !ensure(ParentShareType < ECollectionShareType::CST_All)))
 	{
 		// Bad share type
-		LastError = LOCTEXT("Error_Internal", "There was an internal error.");
+		if (OutError)
+		{
+			*OutError = LOCTEXT("Error_Internal", "There was an internal error.");
+		}
 		return true;
 	}
 
@@ -1530,16 +1955,19 @@ bool FCollectionManager::IsValidParentCollection(FName CollectionName, ECollecti
 	}
 
 	bool bValidParent = true;
-
-	auto IsValidParentCollectionWorker = [&](const FCollectionNameType& InCollectionKey, ECollectionRecursionFlags::Flag InReason) -> FCollectionManagerCache::ERecursiveWorkerFlowControl
+	auto IsValidParentCollectionWorker = [OutError, CollectionName, ShareType, &bValidParent, AvailableCollections=&AvailableCollections]
+	(const FCollectionNameType& InCollectionKey, ECollectionRecursionFlags::Flag InReason) -> FCollectionManagerCache::ERecursiveWorkerFlowControl
 	{
 		const bool bMatchesCollectionBeingReparented = (CollectionName == InCollectionKey.Name && ShareType == InCollectionKey.Type);
 		if (bMatchesCollectionBeingReparented)
 		{
 			bValidParent = false;
-			LastError = (InReason == ECollectionRecursionFlags::Self) 
-				? LOCTEXT("InvalidParent_CannotParentToSelf", "A collection cannot be parented to itself") 
-				: LOCTEXT("InvalidParent_CannotParentToChildren", "A collection cannot be parented to its children");
+			if (OutError)
+			{
+				*OutError = (InReason == ECollectionRecursionFlags::Self)
+						? LOCTEXT("InvalidParent_CannotParentToSelf", "A collection cannot be parented to itself")
+						: LOCTEXT("InvalidParent_CannotParentToChildren", "A collection cannot be parented to its children");
+			}
 			return FCollectionManagerCache::ERecursiveWorkerFlowControl::Stop;
 		}
 
@@ -1547,18 +1975,24 @@ bool FCollectionManager::IsValidParentCollection(FName CollectionName, ECollecti
 		if (!bIsValidChildType)
 		{
 			bValidParent = false;
-			LastError = FText::Format(LOCTEXT("InvalidParent_InvalidChildType", "A {0} collection cannot contain a {1} collection"), ECollectionShareType::ToText(InCollectionKey.Type), ECollectionShareType::ToText(ShareType));
+			if (OutError)
+			{
+				*OutError = FText::Format(LOCTEXT("InvalidParent_InvalidChildType", "A {0} collection cannot contain a {1} collection"), ECollectionShareType::ToText(InCollectionKey.Type), ECollectionShareType::ToText(ShareType));
+			}
 			return FCollectionManagerCache::ERecursiveWorkerFlowControl::Stop;
 		}
 
-		const TSharedRef<FCollection>* const CollectionRefPtr = AvailableCollections.Find(InCollectionKey);
+		const TSharedRef<FCollection>* const CollectionRefPtr = AvailableCollections->Find(InCollectionKey);
 		if (CollectionRefPtr)
 		{
 			const ECollectionStorageMode::Type StorageMode = (*CollectionRefPtr)->GetStorageMode();
 			if (StorageMode == ECollectionStorageMode::Dynamic)
 			{
 				bValidParent = false;
-				LastError = LOCTEXT("InvalidParent_InvalidParentStorageType", "A dynamic collection cannot contain child collections");
+				if (OutError)
+				{
+					*OutError = LOCTEXT("InvalidParent_InvalidParentStorageType", "A dynamic collection cannot contain child collections");
+				}
 				return FCollectionManagerCache::ERecursiveWorkerFlowControl::Stop;
 			}
 		}
@@ -1566,7 +2000,7 @@ bool FCollectionManager::IsValidParentCollection(FName CollectionName, ECollecti
 		return FCollectionManagerCache::ERecursiveWorkerFlowControl::Continue;
 	};
 
-	CollectionCache.RecursionHelper_DoWork(FCollectionNameType(ParentCollectionName, ParentShareType), ECollectionRecursionFlags::SelfAndParents, IsValidParentCollectionWorker);
+	CollectionCache->RecursionHelper_DoWork(InGuard, FCollectionNameType(ParentCollectionName, ParentShareType), ECollectionRecursionFlags::SelfAndParents, IsValidParentCollectionWorker);
 
 	return bValidParent;
 }
@@ -1578,214 +2012,177 @@ void FCollectionManager::HandleFixupRedirectors(ICollectionRedirectorFollower& I
 		return;
 	}
 
-	const double LoadStartTime = FPlatformTime::Seconds();
+	TArray<FCollectionNameType> UpdatedCollections;
+	TArray<FSoftObjectPath> AddedObjects;
+	TArray<FSoftObjectPath> RemovedObjects;
+	{
+		FCollectionLock_Write Guard(Lock);
 
-	TArray<TPair<FSoftObjectPath, FSoftObjectPath>> ObjectsToRename;
+		CollectionCache->UpdateCaches(Guard, ECollectionCacheFlags::Objects);
 
-	// Build up the list of redirected object into rename pairs
-	{
-		const FCollectionObjectsMap& CachedObjects = CollectionCache.GetCachedObjects();
-		for (const auto& CachedObjectInfo : CachedObjects)
+		const double LoadStartTime = FPlatformTime::Seconds();
+
+		TArray<TPair<FSoftObjectPath, FSoftObjectPath>> ObjectsToRename;
+
+		// Build up the list of redirected object into rename pairs
 		{
-			FSoftObjectPath NewObjectPath;
-			if (InRedirectorFollower.FixupObject(CachedObjectInfo.Key, NewObjectPath))
+			const TMap<FSoftObjectPath, TArray<FObjectCollectionInfo>>& CachedObjects = CollectionCache->GetCachedObjects(Guard);
+			for (const TPair<FSoftObjectPath, TArray<FObjectCollectionInfo>>& CachedObjectInfo : CachedObjects)
 			{
-				ObjectsToRename.Emplace(CachedObjectInfo.Key, NewObjectPath);
+				FSoftObjectPath NewObjectPath;
+				if (InRedirectorFollower.FixupObject(CachedObjectInfo.Key, NewObjectPath))
+				{
+					ObjectsToRename.Emplace(CachedObjectInfo.Key, NewObjectPath);
+				}
 			}
 		}
-	}
 
-	TArray<FCollectionNameType> UpdatedCollections;
+		AddedObjects.Reserve(ObjectsToRename.Num());
+		RemovedObjects.Reserve(ObjectsToRename.Num());
 
-	TArray<FSoftObjectPath> AddedObjects;
-	AddedObjects.Reserve(ObjectsToRename.Num());
-
-	TArray<FSoftObjectPath> RemovedObjects;
-	RemovedObjects.Reserve(ObjectsToRename.Num());
-
-	// Handle the rename for each redirected object
-	for (const TPair<FSoftObjectPath, FSoftObjectPath>& ObjectToRename : ObjectsToRename)
-	{
-		AddedObjects.Add(ObjectToRename.Value);
-		RemovedObjects.Add(ObjectToRename.Key);
+		// Handle the rename for each redirected object
+		for (const TPair<FSoftObjectPath, FSoftObjectPath>& ObjectToRename : ObjectsToRename)
+		{
+			AddedObjects.Add(ObjectToRename.Value);
+			RemovedObjects.Add(ObjectToRename.Key);
 
-		ReplaceObjectInCollections(ObjectToRename.Key, ObjectToRename.Value, UpdatedCollections);
-	}
+			ReplaceObjectInCollections(Guard, ObjectToRename.Key, ObjectToRename.Value, UpdatedCollections);
+		}
 
-	if (UpdatedCollections.Num() > 0)
-	{
-		CollectionCache.HandleCollectionChanged();
+		UE_LOG(LogCollectionManager, Log, TEXT( "Fixed up redirectors for %d collections in %0.6f seconds (updated %d objects)" ), AvailableCollections.Num(), FPlatformTime::Seconds() - LoadStartTime, ObjectsToRename.Num());
 
-		// Notify every collection that changed
-		for (const FCollectionNameType& UpdatedCollection : UpdatedCollections)
+		for (const TPair<FSoftObjectPath, FSoftObjectPath>& ObjectToRename : ObjectsToRename)
 		{
-			AssetsRemovedFromCollectionDelegate.Broadcast(UpdatedCollection, RemovedObjects);
-			AssetsAddedToCollectionDelegate.Broadcast(UpdatedCollection, AddedObjects);
+			UE_LOG(LogCollectionManager, Verbose, TEXT( "\tRedirected '%s' to '%s'" ), *ObjectToRename.Key.ToString(), *ObjectToRename.Value.ToString());
 		}
-
-PRAGMA_DISABLE_DEPRECATION_WARNINGS
-		if (AssetsRemovedEvent.IsBound() || AssetsAddedEvent.IsBound())
+		if (UpdatedCollections.Num() > 0)
 		{
-			TArray<FName> RemovedObjectPathNames = UE::SoftObjectPath::Private::ConvertSoftObjectPaths(RemovedObjects);
-			TArray<FName> AddedObjectPathNames = UE::SoftObjectPath::Private::ConvertSoftObjectPaths(AddedObjects);
-			for (const FCollectionNameType& UpdatedCollection : UpdatedCollections)
-			{
-				AssetsRemovedEvent.Broadcast(UpdatedCollection, RemovedObjectPathNames);
-				AssetsAddedEvent.Broadcast(UpdatedCollection, AddedObjectPathNames);
-			}
+			CollectionCache->HandleCollectionChanged(Guard);
 		}
-PRAGMA_ENABLE_DEPRECATION_WARNINGS
 	}
 
-	UE_LOG(LogCollectionManager, Log, TEXT( "Fixed up redirectors for %d collections in %0.6f seconds (updated %d objects)" ), AvailableCollections.Num(), FPlatformTime::Seconds() - LoadStartTime, ObjectsToRename.Num());
-
-	for (const auto& ObjectToRename : ObjectsToRename)
+	// Notify every collection that changed, outside of the lock 
+	for (const FCollectionNameType& UpdatedCollection : UpdatedCollections)
 	{
-		UE_LOG(LogCollectionManager, Verbose, TEXT( "\tRedirected '%s' to '%s'" ), *ObjectToRename.Key.ToString(), *ObjectToRename.Value.ToString());
+		AssetsRemovedFromCollectionDelegate.Broadcast(UpdatedCollection, RemovedObjects);
+		AssetsAddedToCollectionDelegate.Broadcast(UpdatedCollection, AddedObjects);
 	}
 }
 
-bool FCollectionManager::HandleRedirectorsDeleted(TConstArrayView<FSoftObjectPath> ObjectPaths)
+bool FCollectionManager::HandleRedirectorsDeleted(TConstArrayView<FSoftObjectPath> ObjectPaths, FText* OutError)
 {
 	bool bSavedAllCollections = true;
-
-	FTextBuilder AllErrors;
-
 	TArray<FCollectionNameType> UpdatedCollections;
-	TSet<FCollectionNameType> CollectionsToSave;
+	{ 
+		FCollectionLock_Write Guard(Lock);
+		TSet<FCollectionNameType> CollectionsToSave;
+		FTextBuilder ErrorBuilder;
 
-	for (const FSoftObjectPath& ObjectPath : ObjectPaths)
-	{
-		// We don't have a cache for on-disk objects, so we have to do this the slower way and query each collection in turn
-		for (const auto& AvailableCollection : AvailableCollections)
+		for (const FSoftObjectPath& ObjectPath : ObjectPaths)
 		{
-			const FCollectionNameType& CollectionKey = AvailableCollection.Key;
-			const TSharedRef<FCollection>& Collection = AvailableCollection.Value;
-
-			if (Collection->IsRedirectorInCollection(ObjectPath))
+			// We don't have a cache for on-disk objects, so we have to do this the slower way and query each collection in turn
+			for (const TPair<FCollectionNameType, TSharedRef<FCollection>>& AvailableCollection : AvailableCollections)
 			{
-				CollectionsToSave.Add(CollectionKey);
+				const FCollectionNameType& CollectionKey = AvailableCollection.Key;
+				const TSharedRef<FCollection>& Collection = AvailableCollection.Value;
+
+				if (Collection->IsRedirectorInCollection(ObjectPath))
+				{
+					CollectionsToSave.Add(CollectionKey);
+				}
 			}
 		}
-	}
 
-	for (const FCollectionNameType& CollectionKey : CollectionsToSave)
-	{
-		if (TSharedRef<FCollection>* const CollectionRefPtr = AvailableCollections.Find(CollectionKey))
+		for (const FCollectionNameType& CollectionKey : CollectionsToSave)
 		{
-			const TSharedRef<FCollection>& Collection = *CollectionRefPtr;
-
-			FText SaveError;
-			constexpr bool bForceCommitToRevisionControl = true;
-			if (InternalSaveCollection(Collection, SaveError, bForceCommitToRevisionControl))
+			if (TSharedRef<FCollection>* const CollectionRefPtr = AvailableCollections.Find(CollectionKey))
 			{
-				CollectionFileCaches[CollectionKey.Type]->IgnoreFileModification(Collection->GetSourceFilename());
+				const TSharedRef<FCollection>& Collection = *CollectionRefPtr;
 
-				UpdatedCollections.Add(CollectionKey);
-			}
-			else
-			{
-				AllErrors.AppendLine(SaveError);
-				bSavedAllCollections = false;
+				FText SaveError;
+				constexpr bool bForceCommitToRevisionControl = true;
+				if (InternalSaveCollection(Guard, Collection, &SaveError, bForceCommitToRevisionControl))
+				{
+					CollectionFileCaches[CollectionKey.Type]->IgnoreFileModification(Collection->GetSourceFilename());
+
+					UpdatedCollections.Add(CollectionKey);
+				}
+				else
+				{
+					UE_LOG(LogCollectionManager, Error, TEXT("Error saving collection on redirector deletion: %s"), *SaveError.ToString());
+					ErrorBuilder.AppendLine(SaveError);
+					bSavedAllCollections = false;
+				}
 			}
 		}
-	}
-
-	// Notify every collection that changed
-	for (const FCollectionNameType& UpdatedCollection : UpdatedCollections)
-	{
-		AssetsRemovedFromCollectionDelegate.Broadcast(UpdatedCollection, ObjectPaths);
-	}
 
-PRAGMA_DISABLE_DEPRECATION_WARNINGS
-	if( AssetsRemovedEvent.IsBound())
-	{
-		TArray<FName> RemovedObjectPathNames = UE::SoftObjectPath::Private::ConvertSoftObjectPaths(ObjectPaths);
-		for (const FCollectionNameType& UpdatedCollection : UpdatedCollections)
+		if (OutError)
 		{
-			AssetsRemovedEvent.Broadcast(UpdatedCollection, RemovedObjectPathNames);
+			*OutError = ErrorBuilder.ToText();
 		}
 	}
-PRAGMA_ENABLE_DEPRECATION_WARNINGS
 
-	if (!bSavedAllCollections)
+	// Notify every collection that changed, outside of the lock 
+	for (const FCollectionNameType& UpdatedCollection : UpdatedCollections)
 	{
-		LastError = AllErrors.ToText();
+		AssetsRemovedFromCollectionDelegate.Broadcast(UpdatedCollection, ObjectPaths);
 	}
 
 	return bSavedAllCollections;
 }
 
-bool FCollectionManager::HandleRedirectorDeleted(const FSoftObjectPath& ObjectPath)
+bool FCollectionManager::HandleRedirectorDeleted(const FSoftObjectPath& ObjectPath, FText* Error)
 {
-	return HandleRedirectorsDeleted(MakeArrayView(&ObjectPath, 1));
+	return HandleRedirectorsDeleted(MakeArrayView(&ObjectPath, 1), Error);
 }
 
 void FCollectionManager::HandleObjectRenamed(const FSoftObjectPath& OldObjectPath, const FSoftObjectPath& NewObjectPath)
 {
-	TArray<FCollectionNameType> UpdatedCollections; 
-	ReplaceObjectInCollections(OldObjectPath, NewObjectPath, UpdatedCollections);
-
+	TArray<FCollectionNameType> UpdatedCollections;
 	TArray<FSoftObjectPath> AddedObjects;
-	AddedObjects.Add(NewObjectPath);
-
 	TArray<FSoftObjectPath> RemovedObjects;
-	RemovedObjects.Add(OldObjectPath);
-
-	if (UpdatedCollections.Num() > 0)
 	{
-		CollectionCache.HandleCollectionChanged();
-	
-		// Notify every collection that changed
-		for (const FCollectionNameType& UpdatedCollection : UpdatedCollections)
-		{
-			AssetsRemovedFromCollectionDelegate.Broadcast(UpdatedCollection, RemovedObjects);
-			AssetsAddedToCollectionDelegate.Broadcast(UpdatedCollection, AddedObjects);
-		}
+		FCollectionLock_Write Guard(Lock);
+
+		ReplaceObjectInCollections(Guard, OldObjectPath, NewObjectPath, UpdatedCollections);
 
-PRAGMA_DISABLE_DEPRECATION_WARNINGS
-		if( AssetsAddedEvent.IsBound() || AssetsRemovedEvent.IsBound())
+		AddedObjects.Add(NewObjectPath);
+		RemovedObjects.Add(OldObjectPath);
+
+		if (UpdatedCollections.Num() > 0)
 		{
-			TArray<FName> RemovedObjectPathNames = UE::SoftObjectPath::Private::ConvertSoftObjectPaths(RemovedObjects);
-			TArray<FName> AddedObjectPathNames = UE::SoftObjectPath::Private::ConvertSoftObjectPaths(AddedObjects);
-			for (const FCollectionNameType& UpdatedCollection : UpdatedCollections)
-			{
-				AssetsRemovedEvent.Broadcast(UpdatedCollection, RemovedObjectPathNames);
-				AssetsAddedEvent.Broadcast(UpdatedCollection, AddedObjectPathNames);
-			}
+			CollectionCache->HandleCollectionChanged(Guard);
 		}
-PRAGMA_ENABLE_DEPRECATION_WARNINGS
+	}	
+
+	// Notify every collection that changed, outside the lock
+	for (const FCollectionNameType& UpdatedCollection : UpdatedCollections)
+	{
+		AssetsRemovedFromCollectionDelegate.Broadcast(UpdatedCollection, RemovedObjects);
+		AssetsAddedToCollectionDelegate.Broadcast(UpdatedCollection, AddedObjects);
 	}
 }
 
 void FCollectionManager::HandleObjectsDeleted(TConstArrayView<FSoftObjectPath> ObjectPaths)
 {
 	TArray<FCollectionNameType> UpdatedCollections;
-	for (const FSoftObjectPath& ObjectPath : ObjectPaths)
 	{
-		RemoveObjectFromCollections(ObjectPath, UpdatedCollections);
-	}
-
-	if (UpdatedCollections.Num() > 0)
-	{
-		CollectionCache.HandleCollectionChanged();
-	
-		// Notify every collection that changed
-		for (const FCollectionNameType& UpdatedCollection : UpdatedCollections)
+		FCollectionLock_Write Guard(Lock);
+		for (const FSoftObjectPath& ObjectPath : ObjectPaths)
 		{
-			AssetsRemovedFromCollectionDelegate.Broadcast(UpdatedCollection, ObjectPaths);
+			RemoveObjectFromCollections(Guard, ObjectPath, UpdatedCollections);
 		}
 
-PRAGMA_DISABLE_DEPRECATION_WARNINGS
-		if (AssetsRemovedEvent.IsBound())
+		if (UpdatedCollections.Num() > 0)
 		{
-			TArray<FName> RemovedObjectPathNames = UE::SoftObjectPath::Private::ConvertSoftObjectPaths(ObjectPaths);
-			for (const FCollectionNameType& UpdatedCollection : UpdatedCollections)
-			{
-				AssetsRemovedEvent.Broadcast(UpdatedCollection, RemovedObjectPathNames);
-			}
+			CollectionCache->HandleCollectionChanged(Guard);
 		}
-PRAGMA_ENABLE_DEPRECATION_WARNINGS
+	}
+
+	// Notify every collection that changed, outside the lock 
+	for (const FCollectionNameType& UpdatedCollection : UpdatedCollections)
+	{
+		AssetsRemovedFromCollectionDelegate.Broadcast(UpdatedCollection, ObjectPaths);
 	}
 }
 
@@ -1806,20 +2203,27 @@ bool FCollectionManager::TickFileCache(float InDeltaTime)
 		RemoveCollection,
 	};
 
-	bool bDidChangeCollection = false;
-
-	// Process changes that have happened outside of the collection manager
-	for (int32 CacheIdx = 0; CacheIdx < ECollectionShareType::CST_All; ++CacheIdx)
+	// Cached events to fire when we release the lock 
+	TArray<TTuple<ECollectionFileAction, FCollectionNameType>> Events;
 	{
-		const ECollectionShareType::Type ShareType = ECollectionShareType::Type(CacheIdx);
+		// Acquire write lock immediately so we don't need to deal with state change during promotion
+		FCollectionLock_Write Guard(Lock);
 
-		auto& FileCache = CollectionFileCaches[CacheIdx];
-		if (FileCache.IsValid())
+		// Process changes that have happened outside of the collection manager
+		for (int32 CacheIdx = 0; CacheIdx < ECollectionShareType::CST_All; ++CacheIdx)
 		{
+			const ECollectionShareType::Type ShareType = ECollectionShareType::Type(CacheIdx);
+
+			TSharedPtr<DirectoryWatcher::FFileCache>& FileCache = CollectionFileCaches[CacheIdx];
+			if (!FileCache.IsValid())
+			{
+				continue;
+			}
+
 			FileCache->Tick();
 
-			const auto FileCacheChanges = FileCache->GetOutstandingChanges();
-			for (const auto& FileCacheChange : FileCacheChanges)
+			const TArray<DirectoryWatcher::FUpdateCacheTransaction> FileCacheChanges = FileCache->GetOutstandingChanges();
+			for (const DirectoryWatcher::FUpdateCacheTransaction& FileCacheChange : FileCacheChanges)
 			{
 				const FString CollectionFilename = FileCacheChange.Filename.Get();
 				if (FPaths::GetExtension(CollectionFilename) != CollectionExtension)
@@ -1861,10 +2265,9 @@ bool FCollectionManager::TickFileCache(float InDeltaTime)
 						TSharedRef<FCollection> NewCollection = MakeShareable(new FCollection(GetCollectionFilename(CollectionName, ShareType), bUseSCC, ECollectionStorageMode::Static));
 						if (NewCollection->Load(LoadErrorText))
 						{
-							if (AddCollection(NewCollection, ShareType))
+							if (AddCollection(Guard, NewCollection, ShareType))
 							{
-								bDidChangeCollection = true;
-								CollectionCreatedEvent.Broadcast(FCollectionNameType(CollectionName, ShareType));
+								Events.Emplace(CollectionFileAction, FCollectionNameType(CollectionName, ShareType));
 							}
 						}
 						else
@@ -1885,8 +2288,7 @@ bool FCollectionManager::TickFileCache(float InDeltaTime)
 						{
 							if ((*CollectionRefPtr)->Merge(TempCollection))
 							{
-								bDidChangeCollection = true;
-								CollectionUpdatedEvent.Broadcast(FCollectionNameType(CollectionName, ShareType));
+								Events.Emplace(CollectionFileAction, FCollectionNameType(CollectionName, ShareType));
 							}
 						}
 						else
@@ -1901,8 +2303,8 @@ bool FCollectionManager::TickFileCache(float InDeltaTime)
 						TSharedRef<FCollection>* const CollectionRefPtr = AvailableCollections.Find(FCollectionNameType(CollectionName, ShareType));
 						check(CollectionRefPtr); // We tested AvailableCollections.Contains(...) above, so this shouldn't fail
 
-						RemoveCollection(*CollectionRefPtr, ShareType);
-						CollectionDestroyedEvent.Broadcast(FCollectionNameType(CollectionName, ShareType));
+						RemoveCollection(Guard, *CollectionRefPtr, ShareType);
+						Events.Emplace(CollectionFileAction, FCollectionNameType(CollectionName, ShareType));
 					}
 					break;
 
@@ -1911,11 +2313,28 @@ bool FCollectionManager::TickFileCache(float InDeltaTime)
 				}
 			}
 		}
+
+		if (Events.Num() > 0)
+		{
+			CollectionCache->HandleCollectionChanged(Guard);
+		}	
 	}
 
-	if (bDidChangeCollection)
+	// Broadcast events outside the lock
+	for (const TTuple<ECollectionFileAction, FCollectionNameType>& Event : Events)
 	{
-		CollectionCache.HandleCollectionChanged();
+		switch(Event.Key)
+		{
+		case ECollectionFileAction::AddCollection:
+			CollectionCreatedEvent.Broadcast(Event.Value);
+			break;
+		case ECollectionFileAction::MergeCollection:
+			CollectionUpdatedEvent.Broadcast(Event.Value);
+			break;
+		case ECollectionFileAction::RemoveCollection:
+			CollectionDestroyedEvent.Broadcast(Event.Value);
+			break;
+		}
 	}
 
 	return true; // Tick again
@@ -1929,16 +2348,16 @@ void FCollectionManager::LoadCollections()
 	const int32 PrevNumCollections = AvailableCollections.Num();
 	LLM_SCOPE_BYNAME(TEXT("CollectionManager"));
 
-	FRWLock CollectionLock;
+	// This function should only be called during construction, don't acquire a lock here, acquire it for each individual add operation
 	ParallelFor(
 		TEXT("LoadCollections.PF"),
 		ECollectionShareType::CST_All,1,
-		[this, &CollectionLock](int32 CacheIdx)
+		[this](int32 CacheIdx)
 		{
 			const ECollectionShareType::Type ShareType = ECollectionShareType::Type(CacheIdx);
 			const bool bUseSCC = ShouldUseSCC(ShareType);
 			const FString& CollectionFolder = CollectionFolders[CacheIdx];
-			const FString WildCard = FString::Printf(TEXT("%s/*.%s"), *CollectionFolder, *CollectionExtension);
+			const TStringBuilder<256> WildCard{InPlace, CollectionFolder, TEXTVIEW("/*."), CollectionExtension};
 
 			TArray<FString> Filenames;
 			IFileManager::Get().FindFiles(Filenames, *WildCard, true, false);
@@ -1946,7 +2365,7 @@ void FCollectionManager::LoadCollections()
 			ParallelFor(
 				TEXT("LoadCollections.PF"),
 				Filenames.Num(),1,
-				[this, &CollectionLock, &Filenames, &CollectionFolder, bUseSCC, ShareType](int32 FilenameIdx)
+				[this, &Filenames, &CollectionFolder, bUseSCC, ShareType](int32 FilenameIdx)
 				{
 					const FString& BaseFilename = Filenames[FilenameIdx];
 					const FString Filename = CollectionFolder / BaseFilename;
@@ -1955,8 +2374,8 @@ void FCollectionManager::LoadCollections()
 					TSharedRef<FCollection> NewCollection = MakeShareable(new FCollection(Filename, bUseSCC, ECollectionStorageMode::Static));
 					if (NewCollection->Load(LoadErrorText))
 					{
-						FRWScopeLock Lock(CollectionLock, SLT_Write);
-						AddCollection(NewCollection, ShareType);
+						FCollectionLock_Write Guard(Lock);
+						AddCollection(Guard, NewCollection, ShareType);
 					}
 					else
 					{
@@ -1970,7 +2389,8 @@ void FCollectionManager::LoadCollections()
 	);
 
 	// AddCollection is assumed to be adding an empty collection, so also notify that collection cache that the collection has "changed" since loaded collections may not always be empty
-	CollectionCache.HandleCollectionChanged();
+	FCollectionLock_Write Guard(Lock);
+	CollectionCache->HandleCollectionChanged(Guard);
 
 	UE_LOG(LogCollectionManager, Log, TEXT( "Loaded %d collections in %0.6f seconds" ), AvailableCollections.Num() - PrevNumCollections, FPlatformTime::Seconds() - LoadStartTime);
 }
@@ -1987,7 +2407,7 @@ FString FCollectionManager::GetCollectionFilename(const FName& InCollectionName,
 	return CollectionFilename;
 }
 
-bool FCollectionManager::AddCollection(const TSharedRef<FCollection>& CollectionRef, ECollectionShareType::Type ShareType)
+bool FCollectionManager::AddCollection(FCollectionLock_Write& Guard, const TSharedRef<FCollection>& CollectionRef, ECollectionShareType::Type ShareType)
 {
 	if (!ensure(ShareType < ECollectionShareType::CST_All))
 	{
@@ -2003,11 +2423,11 @@ bool FCollectionManager::AddCollection(const TSharedRef<FCollection>& Collection
 	}
 
 	AvailableCollections.Add(CollectionKey, CollectionRef);
-	CollectionCache.HandleCollectionAdded();
+	CollectionCache->HandleCollectionAdded(Guard);
 	return true;
 }
 
-bool FCollectionManager::RemoveCollection(const TSharedRef<FCollection>& CollectionRef, ECollectionShareType::Type ShareType)
+bool FCollectionManager::RemoveCollection(FCollectionLock_Write& Guard, const TSharedRef<FCollection>& CollectionRef, ECollectionShareType::Type ShareType)
 {
 	if (!ensure(ShareType < ECollectionShareType::CST_All))
 	{
@@ -2018,18 +2438,18 @@ bool FCollectionManager::RemoveCollection(const TSharedRef<FCollection>& Collect
 	const FCollectionNameType CollectionKey(CollectionRef->GetCollectionName(), ShareType);
 	if (AvailableCollections.Remove(CollectionKey) > 0)
 	{
-		CollectionCache.HandleCollectionRemoved();
+		CollectionCache->HandleCollectionRemoved(Guard);
 		return true;
 	}
 
 	return false;
 }
 
-void FCollectionManager::RemoveObjectFromCollections(const FSoftObjectPath& ObjectPath, TArray<FCollectionNameType>& OutUpdatedCollections)
+void FCollectionManager::RemoveObjectFromCollections(FCollectionLock_Write& Guard, const FSoftObjectPath& ObjectPath, TArray<FCollectionNameType>& OutUpdatedCollections)
 {
-	const FCollectionObjectsMap& CachedObjects = CollectionCache.GetCachedObjects();
+	const TMap<FSoftObjectPath, TArray<FObjectCollectionInfo>>& CachedObjects = CollectionCache->GetCachedObjects(Guard);
 
-	const auto* ObjectCollectionInfosPtr = CachedObjects.Find(ObjectPath);
+	const TArray<FObjectCollectionInfo>* ObjectCollectionInfosPtr = CachedObjects.Find(ObjectPath);
 	if (!ObjectCollectionInfosPtr)
 	{
 		return;
@@ -2052,11 +2472,12 @@ void FCollectionManager::RemoveObjectFromCollections(const FSoftObjectPath& Obje
 	}
 }
 
-void FCollectionManager::ReplaceObjectInCollections(const FSoftObjectPath& OldObjectPath, const FSoftObjectPath& NewObjectPath, TArray<FCollectionNameType>& OutUpdatedCollections)
+void FCollectionManager::ReplaceObjectInCollections(FCollectionLock_Write& InGuard, const FSoftObjectPath& OldObjectPath, const FSoftObjectPath& NewObjectPath, TArray<FCollectionNameType>& OutUpdatedCollections)
 {
-	const FCollectionObjectsMap& CachedObjects = CollectionCache.GetCachedObjects();
+	CollectionCache->UpdateCaches(InGuard, ECollectionCacheFlags::Objects);
+	const TMap<FSoftObjectPath, TArray<FObjectCollectionInfo>>& CachedObjects = CollectionCache->GetCachedObjects(InGuard);
 
-	const auto* OldObjectCollectionInfosPtr = CachedObjects.Find(OldObjectPath);
+	const TArray<FObjectCollectionInfo>* OldObjectCollectionInfosPtr = CachedObjects.Find(OldObjectPath);
 	if (!OldObjectCollectionInfosPtr)
 	{
 		return;
@@ -2080,11 +2501,11 @@ void FCollectionManager::ReplaceObjectInCollections(const FSoftObjectPath& OldOb
 	}
 }
 
-bool FCollectionManager::InternalSaveCollection(const TSharedRef<FCollection>& CollectionRef, FText& OutError, bool bForceCommitToRevisionControl)
+bool FCollectionManager::InternalSaveCollection(FCollectionLock_Write& InGuard, const TSharedRef<FCollection>& CollectionRef, FText* OutError, bool bForceCommitToRevisionControl)
 {
 	TArray<FText> AdditionalChangelistText;
 
-	// Give game specific editors a chance to add lines
+	// Give game specific editors a chance to add lines - do this under the lock because we don't expect re-entrancy
 	AddToCollectionCheckinDescriptionEvent.Broadcast(CollectionRef->GetCollectionName(), AdditionalChangelistText);
 
 	// Give settings a chance to add lines
@@ -2104,7 +2525,8 @@ bool FCollectionManager::InternalSaveCollection(const TSharedRef<FCollection>& C
 	}
 
 	// Save the collection
-	return CollectionRef->Save(AdditionalChangelistText, OutError, bForceCommitToRevisionControl);
+	FText UnusedError;
+	return CollectionRef->Save(AdditionalChangelistText, OutError ? *OutError : UnusedError, bForceCommitToRevisionControl);
 }
 
 #undef LOCTEXT_NAMESPACE
diff --git a/Engine/Source/Developer/CollectionManager/Private/CollectionManager.h b/Engine/Source/Developer/CollectionManager/Private/CollectionManager.h
--- a/Engine/Source/Developer/CollectionManager/Private/CollectionManager.h
+++ b/Engine/Source/Developer/CollectionManager/Private/CollectionManager.h
@@ -4,14 +4,26 @@
 
 #include "CoreMinimal.h"
 #include "Containers/Ticker.h"
-#include "Misc/Guid.h"
+#include "Collection.h"
 #include "CollectionManagerTypes.h"
 #include "ICollectionManager.h"
-#include "Collection.h"
+#include "Misc/Guid.h"
+#include "Templates/PimplPtr.h"
 
 class ITextFilterExpressionContext;
-
-namespace DirectoryWatcher { class FFileCache; }
+namespace DirectoryWatcher{ class FFileCache; }
+class FCollectionManagerCache;
+
+// Objects wrapping locks to read, write, or begin-reading-then-write (for cache updates) internal state. 
+// Used as internal function parameters to show what lock type must be held to perform the operation and prevent 
+// recursive lock acquisition
+// Functions taking Lock/Lock_Read need to be able to read data but not update caches
+class FCollectionLock;
+class FCollectionLock_Read;
+// Functions taking Lock_RW may need to promote the lock to a write state to update caches
+class FCollectionLock_RW;
+// Functions taking Lock_Write have exclusive access and can update collections as well as update caches
+class FCollectionLock_Write;
 
 /** Collection info for a given object - gives the collection name, as well as the reason this object is considered to be part of this collection */
 struct FObjectCollectionInfo
@@ -34,80 +46,32 @@ struct FObjectCollectionInfo
 	ECollectionRecursionFlags::Flags Reason;
 };
 
+enum class ECollectionCacheFlags
+{
+	None = 0,
+	Names = 1<<0,
+	Objects = 1<<1,
+	Hierarchy = 1<<2,
+	Colors = 1 <<3,
+
+	// Necessary cache updates for calling collection recursion worker
+	RecursionWorker = Names | Hierarchy,
+	All = Names | Objects | Hierarchy | Colors,
+};
+ENUM_CLASS_FLAGS(ECollectionCacheFlags);
+
+
+UE_DEPRECATED(5.5, "These typedefs have been deprecated. Replace them with their concrete types.")
 typedef TMap<FCollectionNameType, TSharedRef<FCollection>> FAvailableCollectionsMap;
+UE_DEPRECATED(5.5, "These typedefs have been deprecated. Replace them with their concrete types.")
 typedef TMap<FGuid, FCollectionNameType> FGuidToCollectionNamesMap;
+UE_DEPRECATED(5.5, "These typedefs have been deprecated. Replace them with their concrete types.")
 typedef TMap<FSoftObjectPath, TArray<FObjectCollectionInfo>> FCollectionObjectsMap;
+UE_DEPRECATED(5.5, "These typedefs have been deprecated. Replace them with their concrete types.")
 typedef TMap<FGuid, TArray<FGuid>> FCollectionHierarchyMap;
+UE_DEPRECATED(5.5, "These typedefs have been deprecated. Replace them with their concrete types.")
 typedef TArray<FLinearColor> FCollectionColorArray;
 
-/** Wraps up the lazy caching of the collection manager */
-class FCollectionManagerCache
-{
-public:
-	FCollectionManagerCache(FAvailableCollectionsMap& InAvailableCollections);
-
-	/** Dirty the parts of the cache that need to change when a collection is added to our collection manager */
-	void HandleCollectionAdded();
-	
-	/** Dirty the parts of the cache that need to change when a collection is removed from our collection manager */
-	void HandleCollectionRemoved();
-
-	/** Dirty the parts of the cache that need to change when a collection is modified */
-	void HandleCollectionChanged();
-
-	/** Access the CachedCollectionNamesFromGuids map, ensuring that it is up-to-date */
-	const FGuidToCollectionNamesMap& GetCachedCollectionNamesFromGuids() const;
-
-	/** Access the CachedObjects map, ensuring that it is up-to-date */
-	const FCollectionObjectsMap& GetCachedObjects() const;
-
-	/** Access the CachedHierarchy map, ensuring that it is up-to-date */
-	const FCollectionHierarchyMap& GetCachedHierarchy() const;
-
-	/** Access the CachedColors array, ensuring that it is up-to-date */
-	const FCollectionColorArray& GetCachedColors() const;
-
-	enum class ERecursiveWorkerFlowControl : uint8
-	{
-		Stop,
-		Continue,
-	};
-
-	typedef TFunctionRef<ERecursiveWorkerFlowControl(const FCollectionNameType&, ECollectionRecursionFlags::Flag)> FRecursiveWorkerFunc;
-
-	void RecursionHelper_DoWork(const FCollectionNameType& InCollectionKey, const ECollectionRecursionFlags::Flags InRecursionMode, FRecursiveWorkerFunc InWorkerFunc) const;
-	ERecursiveWorkerFlowControl RecursionHelper_DoWorkOnParents(const FCollectionNameType& InCollectionKey, FRecursiveWorkerFunc InWorkerFunc) const;
-	ERecursiveWorkerFlowControl RecursionHelper_DoWorkOnChildren(const FCollectionNameType& InCollectionKey, FRecursiveWorkerFunc InWorkerFunc) const;
-
-private:
-	/** Reference to the collections that are currently available in our owner collection manager */
-	FAvailableCollectionsMap& AvailableCollections;
-
-	/** A map of collection GUIDs to their associated collection names */
-	mutable FGuidToCollectionNamesMap CachedCollectionNamesFromGuids_Internal;
-
-	/** A map of object paths to their associated collection info - only objects that are in collections will appear in here */
-	mutable FCollectionObjectsMap CachedObjects_Internal;
-
-	/** A map of parent collection GUIDs to their child collection GUIDs - only collections that have children will appear in here */
-	mutable FCollectionHierarchyMap CachedHierarchy_Internal;
-
-	/** An array of all unique colors currently used by collections */
-	mutable FCollectionColorArray CachedColors_Internal;
-
-	/** Flag to say whether the CachedCollectionNamesFromGuids map is dirty */
-	mutable bool bIsCachedCollectionNamesFromGuidsDirty : 1;
-
-	/** Flag to say whether the CachedObjects map is dirty */
-	mutable bool bIsCachedObjectsDirty : 1;
-
-	/** Flag to say whether the CachedHierarchy map is dirty */
-	mutable bool bIsCachedHierarchyDirty : 1;
-
-	/** Flag to say whether the CachedColors array is dirty */
-	mutable bool bIsCachedColorsDirty : 1;
-};
-
 class FCollectionManager : public ICollectionManager
 {
 public:
@@ -205,12 +169,7 @@ PRAGMA_DISABLE_DEPRECATION_WARNINGS
 	{
 		return HandleObjectDeleted(FSoftObjectPath(ObjectPath));
 	}
-	DECLARE_DERIVED_EVENT( FCollectionManager, ICollectionManager::FAssetsAddedEvent, FAssetsAddedEvent ); 
-	virtual FAssetsAddedEvent& OnAssetsAdded() override { return AssetsAddedEvent; }
-
-	DECLARE_DERIVED_EVENT( FCollectionManager, ICollectionManager::FAssetsRemovedEvent, FAssetsRemovedEvent );
-	virtual FAssetsRemovedEvent& OnAssetsRemoved() override { return AssetsRemovedEvent; }
-PRAGMA_ENABLE_DEPRECATION_WARNINGS
+	PRAGMA_ENABLE_DEPRECATION_WARNINGS
 
 	virtual bool HasCollections() const override;
 	virtual void GetCollections(TArray<FCollectionNameType>& OutCollections) const override;
@@ -230,32 +189,33 @@ PRAGMA_ENABLE_DEPRECATION_WARNINGS
 	virtual void GetCollectionsContainingObjects(const TArray<FSoftObjectPath>& ObjectPaths, TMap<FCollectionNameType, TArray<FSoftObjectPath>>& OutCollectionsAndMatchedObjects, ECollectionRecursionFlags::Flags RecursionMode = ECollectionRecursionFlags::Self) const override;
 	virtual FString GetCollectionsStringForObject(const FSoftObjectPath& ObjectPath, ECollectionShareType::Type ShareType, ECollectionRecursionFlags::Flags RecursionMode = ECollectionRecursionFlags::Self, bool bFullPaths = true) const override;
 	virtual void CreateUniqueCollectionName(const FName& BaseName, ECollectionShareType::Type ShareType, FName& OutCollectionName) const override;
-	virtual bool IsValidCollectionName(const FString& CollectionName, ECollectionShareType::Type ShareType) const override;
-	virtual bool CreateCollection(FName CollectionName, ECollectionShareType::Type ShareType, ECollectionStorageMode::Type StorageMode) override;
-	virtual bool RenameCollection(FName CurrentCollectionName, ECollectionShareType::Type CurrentShareType, FName NewCollectionName, ECollectionShareType::Type NewShareType) override;
-	virtual bool ReparentCollection(FName CollectionName, ECollectionShareType::Type ShareType, FName ParentCollectionName, ECollectionShareType::Type ParentShareType) override;
-	virtual bool DestroyCollection(FName CollectionName, ECollectionShareType::Type ShareType) override;
-	virtual bool AddToCollection(FName CollectionName, ECollectionShareType::Type ShareType, const FSoftObjectPath& ObjectPath) override;
-	virtual bool AddToCollection(FName CollectionName, ECollectionShareType::Type ShareType, TConstArrayView<FSoftObjectPath> ObjectPaths, int32* OutNumAdded = nullptr) override;
-	virtual bool RemoveFromCollection(FName CollectionName, ECollectionShareType::Type ShareType, const FSoftObjectPath& ObjectPath) override;
-	virtual bool RemoveFromCollection(FName CollectionName, ECollectionShareType::Type ShareType, TConstArrayView<FSoftObjectPath> ObjectPaths, int32* OutNumRemoved = nullptr) override;
-	virtual bool SetDynamicQueryText(FName CollectionName, ECollectionShareType::Type ShareType, const FString& InQueryText) override;
-	virtual bool GetDynamicQueryText(FName CollectionName, ECollectionShareType::Type ShareType, FString& OutQueryText) const override;
-	virtual bool TestDynamicQuery(FName CollectionName, ECollectionShareType::Type ShareType, const ITextFilterExpressionContext& InContext, bool& OutResult) const override;
-	virtual bool EmptyCollection(FName CollectionName, ECollectionShareType::Type ShareType) override;
-	virtual bool SaveCollection(FName CollectionName, ECollectionShareType::Type ShareType) override;
-	virtual bool UpdateCollection(FName CollectionName, ECollectionShareType::Type ShareType) override;
-	virtual bool GetCollectionStatusInfo(FName CollectionName, ECollectionShareType::Type ShareType, FCollectionStatusInfo& OutStatusInfo) const override;
+	virtual bool IsValidCollectionName(const FString& CollectionName, ECollectionShareType::Type ShareType, FText* OutError = nullptr) const override;
+	virtual bool CreateCollection(FName CollectionName, ECollectionShareType::Type ShareType, ECollectionStorageMode::Type StorageMode, FText* OutError = nullptr) override;
+	virtual bool RenameCollection(FName CurrentCollectionName, ECollectionShareType::Type CurrentShareType, FName NewCollectionName, ECollectionShareType::Type NewShareType, FText* OutError = nullptr) override;
+	virtual bool ReparentCollection(FName CollectionName, ECollectionShareType::Type ShareType, FName ParentCollectionName, ECollectionShareType::Type ParentShareType, FText* OutError = nullptr) override;
+	virtual bool DestroyCollection(FName CollectionName, ECollectionShareType::Type ShareType, FText* OutError = nullptr) override;
+	virtual bool AddToCollection(FName CollectionName, ECollectionShareType::Type ShareType, const FSoftObjectPath& ObjectPath, FText* OutError = nullptr) override;
+	virtual bool AddToCollection(FName CollectionName, ECollectionShareType::Type ShareType, TConstArrayView<FSoftObjectPath> ObjectPaths, int32* OutNumAdded = nullptr, FText* OutError = nullptr) override;
+	virtual bool RemoveFromCollection(FName CollectionName, ECollectionShareType::Type ShareType, const FSoftObjectPath& ObjectPath, FText* OutError = nullptr) override;
+	virtual bool RemoveFromCollection(FName CollectionName, ECollectionShareType::Type ShareType, TConstArrayView<FSoftObjectPath> ObjectPaths, int32* OutNumRemoved = nullptr, FText* OutError = nullptr) override;
+	virtual bool SetDynamicQueryText(FName CollectionName, ECollectionShareType::Type ShareType, const FString& InQueryText, FText* OutError = nullptr) override;
+	virtual bool GetDynamicQueryText(FName CollectionName, ECollectionShareType::Type ShareType, FString& OutQueryText, FText* OutError = nullptr) const override;
+	virtual bool TestDynamicQuery(FName CollectionName, ECollectionShareType::Type ShareType, const ITextFilterExpressionContext& InContext, bool& OutResult, FText* OutError = nullptr) const override;
+	virtual bool EmptyCollection(FName CollectionName, ECollectionShareType::Type ShareType, FText* OutError = nullptr) override;
+	virtual bool SaveCollection(FName CollectionName, ECollectionShareType::Type ShareType, FText* OutError = nullptr) override;
+	virtual bool UpdateCollection(FName CollectionName, ECollectionShareType::Type ShareType, FText* OutError = nullptr) override;
+	virtual bool GetCollectionStatusInfo(FName CollectionName, ECollectionShareType::Type ShareType, FCollectionStatusInfo& OutStatusInfo, FText* OutError = nullptr) const override;
 	virtual bool HasCollectionColors(TArray<FLinearColor>* OutColors = nullptr) const override;
-	virtual bool GetCollectionColor(FName CollectionName, ECollectionShareType::Type ShareType, TOptional<FLinearColor>& OutColor) const override;
-	virtual bool SetCollectionColor(FName CollectionName, ECollectionShareType::Type ShareType, const TOptional<FLinearColor>& NewColor) override;
-	virtual bool GetCollectionStorageMode(FName CollectionName, ECollectionShareType::Type ShareType, ECollectionStorageMode::Type& OutStorageMode) const override;
-	virtual bool IsObjectInCollection(const FSoftObjectPath& ObjectPath, FName CollectionName, ECollectionShareType::Type ShareType, ECollectionRecursionFlags::Flags RecursionMode = ECollectionRecursionFlags::Self) const override;
-	virtual bool IsValidParentCollection(FName CollectionName, ECollectionShareType::Type ShareType, FName ParentCollectionName, ECollectionShareType::Type ParentShareType) const override;
-	virtual FText GetLastError() const override { return LastError; }
+	virtual bool GetCollectionColor(FName CollectionName, ECollectionShareType::Type ShareType, TOptional<FLinearColor>& OutColor, FText* OutError = nullptr) const override;
+	virtual bool SetCollectionColor(FName CollectionName, ECollectionShareType::Type ShareType, const TOptional<FLinearColor>& NewColor, FText* OutError = nullptr) override;
+	virtual bool GetCollectionStorageMode(FName CollectionName, ECollectionShareType::Type ShareType, ECollectionStorageMode::Type& OutStorageMode, FText* OutError = nullptr) const override;
+	virtual bool IsObjectInCollection(const FSoftObjectPath& ObjectPath, FName CollectionName, ECollectionShareType::Type ShareType, ECollectionRecursionFlags::Flags RecursionMode = ECollectionRecursionFlags::Self, FText* OutError = nullptr) const override;
+	virtual bool IsValidParentCollection(FName CollectionName, ECollectionShareType::Type ShareType, FName ParentCollectionName, ECollectionShareType::Type ParentShareType, FText* OutError) const override;
+	UE_DEPRECATED(5.5, "Deprecated for thread safety reasons.")
+	virtual FText GetLastError() const override { return FText::GetEmpty(); }
 	virtual void HandleFixupRedirectors(ICollectionRedirectorFollower& InRedirectorFollower) override;
-	virtual bool HandleRedirectorDeleted(const FSoftObjectPath& ObjectPath) override;
-	virtual bool HandleRedirectorsDeleted(TConstArrayView<FSoftObjectPath> ObjectPaths) override;
+	virtual bool HandleRedirectorDeleted(const FSoftObjectPath& ObjectPath, FText* OutError = nullptr) override;
+	virtual bool HandleRedirectorsDeleted(TConstArrayView<FSoftObjectPath> ObjectPaths, FText* OutError = nullptr) override;
 	virtual void HandleObjectRenamed(const FSoftObjectPath& OldObjectPath, const FSoftObjectPath& NewObjectPath) override;
 	virtual void HandleObjectDeleted(const FSoftObjectPath& ObjectPath) override;
 	virtual void HandleObjectsDeleted(TConstArrayView<FSoftObjectPath> ObjectPaths) override;
@@ -294,7 +254,7 @@ private:
 	/** Tick this collection manager so it can process any file cache events */
 	bool TickFileCache(float InDeltaTime);
 
-	/** Loads all collection files from disk */
+	/** Loads all collection files from disk. Must only be called from constructor as it does not lock for the full duration. */
 	void LoadCollections();
 
 	/** Returns true if the specified share type requires source control */
@@ -304,29 +264,50 @@ private:
 	FString GetCollectionFilename(const FName& InCollectionName, const ECollectionShareType::Type InCollectionShareType) const;
 
 	/** Adds a collection to the lookup maps */
-	bool AddCollection(const TSharedRef<FCollection>& CollectionRef, ECollectionShareType::Type ShareType);
+	bool AddCollection(FCollectionLock_Write& InGuard, const TSharedRef<FCollection>& CollectionRef, ECollectionShareType::Type ShareType);
 
 	/** Removes a collection from the lookup maps */
-	bool RemoveCollection(const TSharedRef<FCollection>& CollectionRef, ECollectionShareType::Type ShareType);
+	bool RemoveCollection(FCollectionLock_Write& InGuard, const TSharedRef<FCollection>& CollectionRef, ECollectionShareType::Type ShareType);
 
 	/** Removes an object from any collections that contain it */
-	void RemoveObjectFromCollections(const FSoftObjectPath& ObjectPath, TArray<FCollectionNameType>& OutUpdatedCollections);
+	void RemoveObjectFromCollections(FCollectionLock_Write& InGuard, const FSoftObjectPath& ObjectPath, TArray<FCollectionNameType>& OutUpdatedCollections);
 
 	/** Replaces an object with another in any collections that contain it */
-	void ReplaceObjectInCollections(const FSoftObjectPath& OldObjectPath, const FSoftObjectPath& NewObjectPath, TArray<FCollectionNameType>& OutUpdatedCollections);
+	void ReplaceObjectInCollections(
+		FCollectionLock_Write& InGuard, const FSoftObjectPath& OldObjectPath, const FSoftObjectPath& NewObjectPath,
+		TArray<FCollectionNameType>& OutUpdatedCollections);
 
 	/** Internal common functionality for saving a collection
-	 * bForceCommitToRevisionControl - If the collection's storage mode will save it to source control, then bForceCommitToRevisionControl will ensure that it is committed
-	 * after save.  If this is false, then the collection will be left as a modified file which can be advantageous for slow source control servers.
+	 * bForceCommitToRevisionControl - If the collection's storage mode will save it to source control, then
+	 * bForceCommitToRevisionControl will ensure that it is committed after save.  If this is false, then the collection
+	 * will be left as a modified file which can be advantageous for slow source control servers.
+	 */
+	bool InternalSaveCollection(FCollectionLock_Write&, const TSharedRef<FCollection>& CollectionRef, FText* OutError, bool bForceCommitToRevisionControl);
+
+	/* 
+	 * Internal version of IsValidParentCollection to avoid taking lock recursively.
+	 * Cache must be updated for recursion before calling.
 	 */
-	bool InternalSaveCollection(const TSharedRef<FCollection>& CollectionRef, FText& OutError, bool bForceCommitToRevisionControl);
+	bool IsValidParentCollection_Locked(FCollectionLock& InGuard, FName CollectionName, ECollectionShareType::Type ShareType, FName ParentCollectionName, ECollectionShareType::Type ParentShareType, FText* OutError) const;
+
+	/** 
+	 * Check if the given collection exists.
+	 * Using the public API function risks acquiring the lock recursively.
+	 */
+	bool CollectionExists_Locked(FCollectionLock& InGuard, FName CollectionName, ECollectionShareType::Type ShareType) const;
 
 private:
+	/** Required for updating caches as well as write operations to collections */
+	mutable FRWLock Lock;
+
+	/** Cache of collection hierarchy, identity, etc */
+	TPimplPtr<FCollectionManagerCache> CollectionCache;
+
 	/** The folders that contain collections */
 	FString CollectionFolders[ECollectionShareType::CST_All];
 
 	/** The extension used for collection files */
-	FString CollectionExtension;
+	static FStringView CollectionExtension;
 
 	/** Array of file cache instances that are watching for the collection files changing on disk */
 	TSharedPtr<DirectoryWatcher::FFileCache> CollectionFileCaches[ECollectionShareType::CST_All];
@@ -337,24 +318,12 @@ private:
 	/** A map of collection names to FCollection objects */
 	TMap<FCollectionNameType, TSharedRef<FCollection>> AvailableCollections;
 
-	/** The lazily updated cache for this collection manager */
-	FCollectionManagerCache CollectionCache;
-
-	/** The most recent error that occurred */
-	mutable FText LastError;
-
 	/** Event for when assets are added to a collection */
 	FOnAssetsAddedToCollection AssetsAddedToCollectionDelegate;
 
-	UE_DEPRECATED(5.1, "This event has been replaced by AssetsAddedToCollectionEvent")
-	FAssetsAddedEvent AssetsAddedEvent;
-
 	/** Event for when assets are removed from a collection */
 	FOnAssetsRemovedFromCollection AssetsRemovedFromCollectionDelegate;
 
-	UE_DEPRECATED(5.1, "This event has been replaced by AssetsRemovedFromCollectionEvent")
-	FAssetsRemovedEvent AssetsRemovedEvent;
-
 	/** Event for when collections are renamed */
 	FCollectionRenamedEvent CollectionRenamedEvent;
 
diff --git a/Engine/Source/Developer/CollectionManager/Public/ICollectionManager.h b/Engine/Source/Developer/CollectionManager/Public/ICollectionManager.h
--- a/Engine/Source/Developer/CollectionManager/Public/ICollectionManager.h
+++ b/Engine/Source/Developer/CollectionManager/Public/ICollectionManager.h
@@ -43,17 +43,29 @@ public:
 	/** Returns true if the collection exists */
 	virtual bool CollectionExists(FName CollectionName, ECollectionShareType::Type ShareType) const = 0;
 
-	/** Returns a list of asset paths found in the specified collection and share type */
+	/**
+	 * Gets a list of assets in the given collection
+	 * 
+	 * @return true if any objects were found
+	 */
 	virtual bool GetAssetsInCollection(FName CollectionName, ECollectionShareType::Type ShareType, TArray<FSoftObjectPath>& OutAssetPaths, ECollectionRecursionFlags::Flags RecursionMode = ECollectionRecursionFlags::Self) const = 0;
 	UE_DEPRECATED(5.1, "FNames containing full object paths are deprecated. Use FSoftObjectPath instead.")
 	virtual bool GetAssetsInCollection(FName CollectionName, ECollectionShareType::Type ShareType, TArray<FName>& OutAssetPaths, ECollectionRecursionFlags::Flags RecursionMode = ECollectionRecursionFlags::Self) const = 0;
 
-	/** Returns a list of class paths found in the specified collection and share type */
+	/**
+	 * Gets a list of class paths in the given collection
+	 * 
+	 * @return true if any objects were found
+	 */
 	virtual bool GetClassesInCollection(FName CollectionName, ECollectionShareType::Type ShareType, TArray<FTopLevelAssetPath>& ClassPaths, ECollectionRecursionFlags::Flags RecursionMode = ECollectionRecursionFlags::Self) const = 0;
 	UE_DEPRECATED(5.1, "FNames containing class paths are deprecated. Use FTopLevelAssetPath instead.")
 	virtual bool GetClassesInCollection(FName CollectionName, ECollectionShareType::Type ShareType, TArray<FName>& ClassPaths, ECollectionRecursionFlags::Flags RecursionMode = ECollectionRecursionFlags::Self) const = 0;
 
-	/** Returns a list of object paths found in the specified collection and share type */
+	/**
+	 * Gets a list of objects in the given collection
+	 * 
+	 * @return true if any objects were found
+	 */
 	virtual bool GetObjectsInCollection(FName CollectionName, ECollectionShareType::Type ShareType, TArray<FSoftObjectPath>& ObjectPaths, ECollectionRecursionFlags::Flags RecursionMode = ECollectionRecursionFlags::Self) const = 0;
 	UE_DEPRECATED(5.1, "FNames containing full object paths are deprecated. Use FSoftObjectPath instead.")
 	virtual bool GetObjectsInCollection(FName CollectionName, ECollectionShareType::Type ShareType, TArray<FName>& ObjectPaths, ECollectionRecursionFlags::Flags RecursionMode = ECollectionRecursionFlags::Self) const = 0;
@@ -78,80 +90,115 @@ public:
 	UE_DEPRECATED(5.1, "FNames containing full object paths are deprecated. Use FSoftObjectPath instead.")
 	virtual FString GetCollectionsStringForObject(FName ObjectPath, ECollectionShareType::Type ShareType, ECollectionRecursionFlags::Flags RecursionMode = ECollectionRecursionFlags::Self, bool bFullPaths = true) const = 0;
 
-	/** Creates a unique collection name for the given type taking the form BaseName+(unique number) */
+	/** 
+	 * Creates a unique collection name for the given type taking the form BaseName+(unique number) 
+	 * 
+	 * @warning This function is threadsafe but it does not secure the name from another thread creating the collection in between calls to the ICollectionManager API.
+	 */
 	virtual void CreateUniqueCollectionName(const FName& BaseName, ECollectionShareType::Type ShareType, FName& OutCollectionName) const = 0;
 
-	/** Returns whether or not the given collection name is valid. If false, GetLastError will return a human readable string description of the error. */
-	virtual bool IsValidCollectionName(const FString& CollectionName, ECollectionShareType::Type ShareType) const = 0;
+	/**
+	 * Check whether the given name is valid for a new collection.
+	 *
+	 * @param CollectionName Potential collection name
+	 * @param ShareType Type of collection to check for conflicts against. CST_All checks against all collcetion types.
+	 * @param OutError Optional human readable explanation of why the operation failed
+	 * @return true If a collection can be created with the name.
+	 * @return false If the name contains invalid characters or conflicts with an existing collection of the same type.
+	 */
+	virtual bool IsValidCollectionName(
+		const FString& CollectionName, ECollectionShareType::Type ShareType, FText* OutError = nullptr) const = 0;
 
 	/**
-	  * Adds a collection to the asset registry. A .collection file will be added to disk.
-	  *
-	  * @param CollectionName The name of the new collection
-	  * @param ShareType The way the collection is shared.
-	  * @param StorageMode How does this collection store its objects? (static or dynamic)
-	  * @return true if the add was successful. If false, GetLastError will return a human readable string description of the error.
-	  */
-	virtual bool CreateCollection(FName CollectionName, ECollectionShareType::Type ShareType, ECollectionStorageMode::Type StorageMode) = 0;
+	 * Creates a new collection. A .collection file will be added to disk.
+	 *
+	 * @param CollectionName The name of the new collection
+	 * @param ShareType The way the collection is shared.
+	 * @param StorageMode How does this collection store its objects? (static or dynamic)
+	 * @param OutError Optional human readable explanation of why the operation failed
+	 * @return true if the add was successful.
+	 */
+	virtual bool CreateCollection(
+		FName CollectionName, ECollectionShareType::Type ShareType, ECollectionStorageMode::Type StorageMode,
+		FText* OutError = nullptr) = 0;
 
 	/**
-	  * Renames a collection. A .collection file will be added to disk and a .collection file will be removed.
-	  *
-	  * @param CurrentCollectionName The current name of the collection.
-	  * @param CurrentShareType The current way the collection is shared.
-	  * @param NewCollectionName The new name of the collection.
-	  * @param NewShareType The new way the collection is shared.
-	  * @return true if the rename was successful. If false, GetLastError will return a human readable string description of the error.
-	  */
-	virtual bool RenameCollection(FName CurrentCollectionName, ECollectionShareType::Type CurrentShareType, FName NewCollectionName, ECollectionShareType::Type NewShareType) = 0;
+	 * Renames a collection. A .collection file will be added to disk and a .collection file will be removed.
+	 *
+	 * @param CurrentCollectionName The current name of the collection.
+	 * @param CurrentShareType The current way the collection is shared.
+	 * @param NewCollectionName The new name of the collection.
+	 * @param NewShareType The new way the collection is shared.
+	 * @param OutError Optional human readable explanation of why the operation failed
+	 * @return true if the rename was successful.
+	 */
+	virtual bool RenameCollection(
+		FName CurrentCollectionName, ECollectionShareType::Type CurrentShareType, FName NewCollectionName,
+		ECollectionShareType::Type NewShareType, FText* OutError = nullptr) = 0;
 
 	/**
-	  * Re-parents a collection. The parent collection may be re-saved if it's too old to have a stable GUID.
-	  *
-	  * @param CollectionName The name of the collection to set the parent of.
-	  * @param ShareType The way the collection is shared.
-	  * @param ParentCollectionName The new parent of the collection, or None to have the collection become a root collection.
-	  * @param ParentShareType The way the new parent collection is shared.
-	  * @return true if the re-parent was successful. If false, GetLastError will return a human readable string description of the error.
-	  */
-	virtual bool ReparentCollection(FName CollectionName, ECollectionShareType::Type ShareType, FName ParentCollectionName, ECollectionShareType::Type ParentShareType) = 0;
+	 * Re-parents a collection. The parent collection may be re-saved if it's too old to have a stable GUID.
+	 *
+	 * @param CollectionName The name of the collection to set the parent of.
+	 * @param ShareType The way the collection is shared.
+	 * @param ParentCollectionName The new parent of the collection, or None to have the collection become a root
+	 * collection.
+	 * @param ParentShareType The way the new parent collection is shared.
+	 * @param OutError Optional human readable explanation of why the operation failed
+	 * @return true if the re-parent was successful.
+	 */
+	virtual bool ReparentCollection(
+		FName CollectionName, ECollectionShareType::Type ShareType, FName ParentCollectionName,
+		ECollectionShareType::Type ParentShareType, FText* OutError = nullptr) = 0;
 
 	/**
-	  * Removes a collection to the asset registry. A .collection file will be deleted from disk.
-	  *
-	  * @param CollectionName The name of the collection to remove.
-	  * @param ShareType The way the collection is shared.
-	  * @return true if the remove was successful
-	  */
-	virtual bool DestroyCollection(FName CollectionName, ECollectionShareType::Type ShareType) = 0;
+	 * Removes a collection to the asset registry. A .collection file will be deleted from disk.
+	 *
+	 * @param CollectionName The name of the collection to remove.
+	 * @param ShareType The way the collection is shared.
+	 * @param OutError Optional human readable explanation of why the operation failed
+	 * @return true if the remove was successful
+	 */
+	virtual bool DestroyCollection(FName CollectionName, ECollectionShareType::Type ShareType, FText* OutError = nullptr) = 0;
 
 	/**
-	  * Adds an asset to the specified collection.
-	  *
-	  * @param CollectionName The collection in which to add the asset
-	  * @param ShareType The way the collection is shared.
-	  * @param ObjectPath the ObjectPath of the asset to add.
-	  * @param OutNumAdded if non-NULL, the number of objects successfully added to the collection
-	  * @return true if the add was successful. If false, GetLastError will return a human readable string description of the error.
-	  */
-	virtual bool AddToCollection(FName CollectionName, ECollectionShareType::Type ShareType, const FSoftObjectPath& ObjectPath) = 0;
-	virtual bool AddToCollection(FName CollectionName, ECollectionShareType::Type ShareType, TConstArrayView<FSoftObjectPath> ObjectPaths, int32* OutNumAdded = nullptr) = 0;
+	 * Adds an asset to the specified collection.
+	 *
+	 * @param CollectionName The collection in which to add the asset
+	 * @param ShareType The way the collection is shared.
+	 * @param ObjectPath the ObjectPath of the asset to add.
+	 * @param OutNumAdded if non-NULL, the number of objects successfully added to the collection
+	 * @param OutError Optional human readable explanation of why the operation failed
+	 * @return true if the add was successful.
+	 */
+	virtual bool AddToCollection(
+		FName CollectionName, ECollectionShareType::Type ShareType, const FSoftObjectPath& ObjectPath,
+		FText* OutError = nullptr) = 0;
+	virtual bool AddToCollection(
+		FName CollectionName, ECollectionShareType::Type ShareType, TConstArrayView<FSoftObjectPath> ObjectPaths,
+		int32* OutNumAdded = nullptr, FText* OutError = nullptr) = 0;
+
 	UE_DEPRECATED(5.1, "FNames containing full object paths are deprecated. Use FSoftObjectPath instead.")
 	virtual bool AddToCollection(FName CollectionName, ECollectionShareType::Type ShareType, FName ObjectPath) = 0;
 	UE_DEPRECATED(5.1, "FNames containing full object paths are deprecated. Use FSoftObjectPath instead.")
 	virtual bool AddToCollection(FName CollectionName, ECollectionShareType::Type ShareType, const TArray<FName>& ObjectPaths, int32* OutNumAdded = nullptr) = 0;
 
 	/**
-	  * Removes the asset from the specified collection.
-	  *
-	  * @param CollectionName The collection from which to remove the asset
-	  * @param ShareType The way the collection is shared.
-	  * @param ObjectPath the ObjectPath of the asset to remove.
-	  * @param OutNumRemoved if non-NULL, the number of objects successfully removed from the collection
-	  * @return true if the remove was successful. If false, GetLastError will return a human readable string description of the error.
-	  */
-	virtual bool RemoveFromCollection(FName CollectionName, ECollectionShareType::Type ShareType, const FSoftObjectPath& ObjectPath) = 0;
-	virtual bool RemoveFromCollection(FName CollectionName, ECollectionShareType::Type ShareType, TConstArrayView<FSoftObjectPath> ObjectPaths, int32* OutNumRemoved = nullptr) = 0;
+	 * Removes the asset from the specified collection.
+	 *
+	 * @param CollectionName The collection from which to remove the asset
+	 * @param ShareType The way the collection is shared.
+	 * @param ObjectPath the ObjectPath of the asset to remove.
+	 * @param OutNumRemoved if non-NULL, the number of objects successfully removed from the collection
+	 * @param OutError Optional human readable explanation of why the operation failed
+	 * @return true if the remove was successful.
+	 */
+	virtual bool RemoveFromCollection(
+		FName CollectionName, ECollectionShareType::Type ShareType, const FSoftObjectPath& ObjectPath,
+		FText* OutError = nullptr) = 0;
+	virtual bool RemoveFromCollection(
+		FName CollectionName, ECollectionShareType::Type ShareType, TConstArrayView<FSoftObjectPath> ObjectPaths,
+		int32* OutNumRemoved = nullptr, FText* OutError = nullptr) = 0;
 	UE_DEPRECATED(5.1, "FNames containing full object paths are deprecated. Use FSoftObjectPath instead.")
 	virtual bool RemoveFromCollection(FName CollectionName, ECollectionShareType::Type ShareType, FName ObjectPath) = 0;
 	UE_DEPRECATED(5.1, "FNames containing full object paths are deprecated. Use FSoftObjectPath instead.")
@@ -163,9 +210,12 @@ public:
 	 * @param CollectionName The collection to set the query on.
 	 * @param ShareType The way the collection is shared.
 	 * @param InQueryText The new query to set.
-	 * @return true if the set was successful. If false, GetLastError will return a human readable string description of the error.
+	 * @param OutError Optional human readable explanation of why the operation failed
+	 * @return true if the set was successful.
 	 */
-	virtual bool SetDynamicQueryText(FName CollectionName, ECollectionShareType::Type ShareType, const FString& InQueryText) = 0;
+	virtual bool SetDynamicQueryText(
+		FName CollectionName, ECollectionShareType::Type ShareType, const FString& InQueryText,
+		FText* OutError = nullptr) = 0;
 
 	/**
 	 * Gets the dynamic query text for the specified collection. 
@@ -173,9 +223,12 @@ public:
 	 * @param CollectionName The collection to get the query from.
 	 * @param ShareType The way the collection is shared.
 	 * @param OutQueryText Filled with the query text.
-	 * @return true if the get was successful. If false, GetLastError will return a human readable string description of the error.
+	 * @param OutError Optional human readable explanation of why the operation failed
+	 * @return true if the get was successful.
 	 */
-	virtual bool GetDynamicQueryText(FName CollectionName, ECollectionShareType::Type ShareType, FString& OutQueryText) const = 0;
+	virtual bool GetDynamicQueryText(
+		FName CollectionName, ECollectionShareType::Type ShareType, FString& OutQueryText,
+		FText* OutError = nullptr) const = 0;
 
 	/**
 	 * Tests the dynamic query for the specified collection against the context provided.
@@ -184,52 +237,64 @@ public:
 	 * @param ShareType The way the collection is shared.
 	 * @param InContext The context to test against.
 	 * @param OutResult Filled with the result of the query.
-	 * @return true if the get was successful. If false, GetLastError will return a human readable string description of the error.
+	 * @param OutError Optional human readable explanation of why the operation failed
+	 * @return true if the get was successful.
 	 */
-	virtual bool TestDynamicQuery(FName CollectionName, ECollectionShareType::Type ShareType, const ITextFilterExpressionContext& InContext, bool& OutResult) const = 0;
+	virtual bool TestDynamicQuery(
+		FName CollectionName, ECollectionShareType::Type ShareType, const ITextFilterExpressionContext& InContext,
+		bool& OutResult, FText* OutError = nullptr) const = 0;
 
 	/**
-	  * Removes all assets from the specified collection.
-	  *
-	  * @param CollectionName The collection to empty
-	  * @param ShareType The way the collection is shared.
-	  * @return true if the clear was successful. If false, GetLastError will return a human readable string description of the error.
-	  */
-	virtual bool EmptyCollection(FName CollectionName, ECollectionShareType::Type ShareType) = 0;
+	 * Removes all assets from the specified collection.
+	 *
+	 * @param CollectionName The collection to empty
+	 * @param ShareType The way the collection is shared.
+	 * @param OutError Optional human readable explanation of why the operation failed
+	 * @return true if the clear was successful.
+	 */
+	virtual bool EmptyCollection(FName CollectionName, ECollectionShareType::Type ShareType, FText* OutError = nullptr) = 0;
 
 	/**
-	  * Save the collection (if dirty) and check it into source control (if under SCC control)
-	  *
-	  * Note: Generally you won't need to save collections manually as the collection manager takes care of that as objects and added/removed, etc. 
-	  *	      however, you may want to manually save a collection if a previous save attempt failed (and you've since corrected the issue), or if 
-	  *       the collection contains redirected object references that you'd like to save to disk.
-	  *
-	  * @param CollectionName The collection to save
-	  * @param ShareType The way the collection is shared.
-	  * @return true if the save was successful. If false, GetLastError will return a human readable string description of the error.
-	  */
-	virtual bool SaveCollection(FName CollectionName, ECollectionShareType::Type ShareType) = 0;
+	 * Save the collection (if dirty) and check it into source control (if under SCC control)
+	 *
+	 * Note: Generally you won't need to save collections manually as the collection manager takes
+	 * care of that as objects and added/removed, etc. however, you may want to manually save a
+	 * collection if a previous save attempt failed (and you've since corrected the issue), or if
+	 * the collection contains redirected object references that you'd like to save to disk.
+	 *
+	 * @param CollectionName The collection to save
+	 * @param ShareType The way the collection is shared.
+	 * @param OutError Optional human readable explanation of why the operation failed
+	 * @return true if the save was successful.
+	 */
+	virtual bool SaveCollection(FName CollectionName, ECollectionShareType::Type ShareType, FText* OutError = nullptr) = 0;
 
 	/**
-	  * Update the collection to make sure it's using the latest version from source control (if under SCC control)
-	  *
-	  * Note: Generally you won't need to update collections manually as the collection manager takes care of that as collections are saved to disk.
-	  *
-	  * @param CollectionName The collection to update
-	  * @param ShareType The way the collection is shared.
-	  * @return true if the update was successful. If false, GetLastError will return a human readable string description of the error.
-	  */
-	virtual bool UpdateCollection(FName CollectionName, ECollectionShareType::Type ShareType) = 0;
+	 * Update the collection to make sure it's using the latest version from source control (if
+	 * under SCC control)
+	 *
+	 * Note: Generally you won't need to update collections manually as the collection manager takes
+	 * care of that as collections are saved to disk.
+	 *
+	 * @param CollectionName The collection to update
+	 * @param ShareType The way the collection is shared.
+	 * @param OutError Optional human readable explanation of why the operation failed
+	 * @return true if the update was successful.
+	 */
+	virtual bool UpdateCollection(FName CollectionName, ECollectionShareType::Type ShareType, FText* OutError = nullptr) = 0;
 
 	/**
-	  * Gets the status info for the specified collection
-	  *
-	  * @param CollectionName The collection to get the status info for
-	  * @param ShareType The way the collection is shared.
-	  * @param OutStatusInfo The status info to populate.
-	  * @return true if the status info was filled in. If false, GetLastError will return a human readable string description of the error.
-	  */
-	virtual bool GetCollectionStatusInfo(FName CollectionName, ECollectionShareType::Type ShareType, FCollectionStatusInfo& OutStatusInfo) const = 0;
+	 * Gets the status info for the specified collection
+	 *
+	 * @param CollectionName The collection to get the status info for
+	 * @param ShareType The way the collection is shared.
+	 * @param OutStatusInfo The status info to populate.
+	 * @param OutError Optional human readable explanation of why the operation failed
+	 * @return true if the status info was filled in.
+	 */
+	virtual bool GetCollectionStatusInfo(
+		FName CollectionName, ECollectionShareType::Type ShareType, FCollectionStatusInfo& OutStatusInfo,
+		FText* OutError = nullptr) const = 0;
 
 	/**
 	 * Checks whether any collections have a custom color set, optionally retrieving the list of in-use colors
@@ -245,9 +310,12 @@ public:
 	 * @param CollectionName The collection to get the color for
 	 * @param ShareType The way the collection is shared.
 	 * @param OutColor The optional color to populate.
-	 * @return true if the collection was found, and the color was optionally filled in. If false, GetLastError will return a human readable string description of the error.
+	 * @param OutError Optional human readable explanation of why the operation failed
+	 * @return true if the collection was found, and the color was optionally filled in.
 	 */
-	virtual bool GetCollectionColor(FName CollectionName, ECollectionShareType::Type ShareType, TOptional<FLinearColor>& OutColor) const = 0;
+	virtual bool GetCollectionColor(
+		FName CollectionName, ECollectionShareType::Type ShareType, TOptional<FLinearColor>& OutColor,
+		FText* OutError = nullptr) const = 0;
 
 	/**
 	 * Sets the optional color for the specified collection
@@ -255,9 +323,12 @@ public:
 	 * @param CollectionName The collection to set the color for
 	 * @param ShareType The way the collection is shared.
 	 * @param NewColor The optional color to set.
-	 * @return true if the collection was found, and the color was set. If false, GetLastError will return a human readable string description of the error.
+	 * @param OutError Optional human readable explanation of why the operation failed
+	 * @return true if the collection was found, and the color was set.
 	 */
-	virtual bool SetCollectionColor(FName CollectionName, ECollectionShareType::Type ShareType, const TOptional<FLinearColor>& NewColor) = 0;
+	virtual bool SetCollectionColor(
+		FName CollectionName, ECollectionShareType::Type ShareType, const TOptional<FLinearColor>& NewColor,
+		FText* OutError = nullptr) = 0;
 
 	/**
 	 * Gets the method by which the specified collection stores its objects (static or dynamic)
@@ -265,9 +336,12 @@ public:
 	 * @param CollectionName The collection to get the storage mode for
 	 * @param ShareType The way the collection is shared.
 	 * @param OutStorageMode The variable to populate.
-	 * @return true if the status info was filled in. If false, GetLastError will return a human readable string description of the error.
+	 * @param OutError Optional human readable explanation of why the operation failed
+	 * @return true if the status info was filled in.
 	 */
-	virtual bool GetCollectionStorageMode(FName CollectionName, ECollectionShareType::Type ShareType, ECollectionStorageMode::Type& OutStorageMode) const = 0;
+	virtual bool GetCollectionStorageMode(
+		FName CollectionName, ECollectionShareType::Type ShareType, ECollectionStorageMode::Type& OutStorageMode,
+		FText* OutError = nullptr) const = 0;
 
 	/**
 	 * Check to see if the given object exists in the given collection
@@ -275,24 +349,33 @@ public:
 	 * @param ObjectPath The object to test for its existence in the collection.
 	 * @param CollectionName The collection to test.
 	 * @param ShareType The way the collection is shared.
+	 * @param OutError Optional human readable explanation of why the operation failed
 	 * @return true if the object is in the collection.
 	 */
-	virtual bool IsObjectInCollection(const FSoftObjectPath& ObjectPath, FName CollectionName, ECollectionShareType::Type ShareType, ECollectionRecursionFlags::Flags RecursionMode = ECollectionRecursionFlags::Self) const = 0;
+	virtual bool IsObjectInCollection(
+		const FSoftObjectPath& ObjectPath, FName CollectionName, ECollectionShareType::Type ShareType,
+		ECollectionRecursionFlags::Flags RecursionMode = ECollectionRecursionFlags::Self,
+		FText* OutError = nullptr) const = 0;
 	UE_DEPRECATED(5.1, "FNames containing full object paths are deprecated. Use FSoftObjectPath instead.")
 	virtual bool IsObjectInCollection(FName ObjectPath, FName CollectionName, ECollectionShareType::Type ShareType, ECollectionRecursionFlags::Flags RecursionMode = ECollectionRecursionFlags::Self) const = 0;
 
 	/**
-	  * Check to see if the given collection is valid to be used as the parent of another collection. A collection may not be parented to itself, nor any of its current children.
-	  *
-	  * @param CollectionName The name of the collection to check the parent of.
-	  * @param ShareType The way the collection is shared.
-	  * @param ParentCollectionName The name of the parent collection to test against.
-	  * @param ParentShareType The way the new parent collection is shared.
-	  * @return true if the parent is considered valid for this collection. If false, GetLastError will return a human readable string describing why the parent is invalid.
-	  */
-	virtual bool IsValidParentCollection(FName CollectionName, ECollectionShareType::Type ShareType, FName ParentCollectionName, ECollectionShareType::Type ParentShareType) const = 0;
+	 * Check to see if the given collection is valid to be used as the parent of another collection.
+	 * A collection may not be parented to itself, nor any of its current children.
+	 *
+	 * @param CollectionName The name of the collection to check the parent of.
+	 * @param ShareType The way the collection is shared.
+	 * @param ParentCollectionName The name of the parent collection to test against.
+	 * @param ParentShareType The way the new parent collection is shared.
+	 * @param OutError Optional human readable explanation of why the operation failed
+	 * @return true if the parent is considered valid for this collection.
+	 */
+	virtual bool IsValidParentCollection(
+		FName CollectionName, ECollectionShareType::Type ShareType, FName ParentCollectionName,
+		ECollectionShareType::Type ParentShareType, FText* OutError = nullptr) const = 0;
 
 	/** Returns the most recent error. */
+	UE_DEPRECATED(5.5, "Deprecated for thread safety reasons. Functions which populated this value now take an optional FText* out parameter.")
 	virtual FText GetLastError() const = 0;
 
 	/**
@@ -305,7 +388,7 @@ public:
 	 * Called to notify the collections that a redirector has been deleted and that they should ensure their on-disk representation is re-saved with the fixed up in-memory version
 	 * @return true if all of the collections that were referencing this redirector could be re-saved, false otherwise
 	 */
-	virtual bool HandleRedirectorDeleted(const FSoftObjectPath& ObjectPath) = 0;
+	virtual bool HandleRedirectorDeleted(const FSoftObjectPath& ObjectPath, FText* OutError = nullptr) = 0;
 	UE_DEPRECATED(5.1, "FNames containing full object paths are deprecated. Use FSoftObjectPath instead.")
 	virtual bool HandleRedirectorDeleted(const FName& ObjectPath) = 0;
 
@@ -313,7 +396,7 @@ public:
 	 * Called to notify the collections that redirectors have been deleted and that they should ensure their on-disk representation is re-saved with the fixed up in-memory version
 	 * @return true if all of the collections that were referencing these redirectors could be re-saved, false otherwise
 	 */
-	virtual bool HandleRedirectorsDeleted(TConstArrayView<FSoftObjectPath> ObjectPaths) = 0;
+	virtual bool HandleRedirectorsDeleted(TConstArrayView<FSoftObjectPath> ObjectPaths, FText* OutError = nullptr) = 0;
 
 	/** Called to notify the collections that an object has been renamed or moved */
 	virtual void HandleObjectRenamed(const FSoftObjectPath& OldObjectPath, const FSoftObjectPath& NewObjectPath) = 0;
@@ -340,18 +423,10 @@ public:
 	DECLARE_MULTICAST_DELEGATE_TwoParams(FOnAssetsAddedToCollection, const FCollectionNameType&, TConstArrayView<FSoftObjectPath> );
 	virtual FOnAssetsAddedToCollection& OnAssetsAddedToCollection() = 0;
 
-	DECLARE_EVENT_TwoParams( ICollectionManager, FAssetsAddedEvent, const FCollectionNameType&, const TArray<FName>& );
-	UE_DEPRECATED(5.1, "Names containing full asset paths are deprecated, use OnAssetsAddedToCollection instead.")
-	virtual FAssetsAddedEvent& OnAssetsAdded() = 0;
-
 	/** Event for when assets are removed from a collection */
 	DECLARE_MULTICAST_DELEGATE_TwoParams( FOnAssetsRemovedFromCollection, const FCollectionNameType&, TConstArrayView<FSoftObjectPath> );
 	virtual FOnAssetsRemovedFromCollection& OnAssetsRemovedFromCollection() = 0;
 
-	DECLARE_EVENT_TwoParams( ICollectionManager, FAssetsRemovedEvent, const FCollectionNameType&, const TArray<FName>& );
-	UE_DEPRECATED(5.1, "Names containing full asset paths are deprecated, use OnAssetsRemovedFromCollection instead.")
-	virtual FAssetsRemovedEvent& OnAssetsRemoved() = 0;
-
 	/** Event for when collections are renamed */
 	DECLARE_EVENT_TwoParams( ICollectionManager, FCollectionRenamedEvent, const FCollectionNameType&, const FCollectionNameType& );
 	virtual FCollectionRenamedEvent& OnCollectionRenamed() = 0;
@@ -364,7 +439,10 @@ public:
 	DECLARE_EVENT_OneParam( ICollectionManager, FCollectionUpdatedEvent, const FCollectionNameType& );
 	virtual FCollectionUpdatedEvent& OnCollectionUpdated() = 0;
 
-	/** When a collection checkin happens, use this event to add additional text to the changelist description */
+	/** 
+	 * When a collection checkin happens, use this event to add additional text to the changelist description.
+	 * Delegates bound to this event must be safe to call on any thread. 
+	 */
 	DECLARE_EVENT_TwoParams( ICollectionManager, FAddToCollectionCheckinDescriptionEvent, const FName& /*CollectionName*/, TArray<FText>& /* OutAdditionalLines*/);
 	virtual FAddToCollectionCheckinDescriptionEvent& OnAddToCollectionCheckinDescriptionEvent() = 0;
 };
diff --git a/Engine/Source/Editor/ContentBrowser/Private/SCollectionView.cpp b/Engine/Source/Editor/ContentBrowser/Private/SCollectionView.cpp
--- a/Engine/Source/Editor/ContentBrowser/Private/SCollectionView.cpp
+++ b/Engine/Source/Editor/ContentBrowser/Private/SCollectionView.cpp
@@ -1052,7 +1052,8 @@ void SCollectionView::DeleteCollectionItems( const TArray<TSharedPtr<FCollection
 	FCollectionManagerModule& CollectionManagerModule = FCollectionManagerModule::GetModule();
 	for (const TSharedPtr<FCollectionItem>& ItemToDelete : ItemsToDelete)
 	{
-		if (CollectionManagerModule.Get().DestroyCollection(ItemToDelete->CollectionName, ItemToDelete->CollectionType))
+		FText Error;
+		if (CollectionManagerModule.Get().DestroyCollection(ItemToDelete->CollectionName, ItemToDelete->CollectionType, &Error))
 		{
 			if (PreviouslySelectedItems.Contains(ItemToDelete))
 			{
@@ -1065,7 +1066,7 @@ void SCollectionView::DeleteCollectionItems( const TArray<TSharedPtr<FCollection
 			const FVector2f& CursorPos = FSlateApplication::Get().GetCursorPos();
 			FSlateRect MessageAnchor(CursorPos.X, CursorPos.Y, CursorPos.X, CursorPos.Y);
 			ContentBrowserUtils::DisplayMessage(
-				FText::Format( LOCTEXT("CollectionDestroyFailed", "Failed to destroy collection. {0}"), CollectionManagerModule.Get().GetLastError() ),
+				FText::Format( LOCTEXT("CollectionDestroyFailed", "Failed to destroy collection. {0}"), Error),
 				MessageAnchor,
 				CollectionTreePtr.ToSharedRef()
 				);
@@ -1232,12 +1233,13 @@ FReply SCollectionView::HandleDragDropOnCollectionTree(const FGeometry& Geometry
 		// Reparent all of the collections in the drag drop so that they are root level items
 		for (const FCollectionNameType& NewChildCollection : DragDropOp->Collections)
 		{
+			FText Error;
 			if (!CollectionManagerModule.Get().ReparentCollection(
 					NewChildCollection.Name, NewChildCollection.Type,
-					NAME_None, ECollectionShareType::CST_All
+					NAME_None, ECollectionShareType::CST_All, &Error
 					))
 			{
-				ContentBrowserUtils::DisplayMessage(CollectionManagerModule.Get().GetLastError(), Geometry.GetLayoutBoundingRect(), SharedThis(this));
+				ContentBrowserUtils::DisplayMessage(Error, Geometry.GetLayoutBoundingRect(), SharedThis(this));
 			}
 		}
 
@@ -1271,14 +1273,14 @@ bool SCollectionView::ValidateDragDropOnCollectionItem(TSharedRef<FCollectionIte
 		bIsValidDrag = true;
 		for (const FCollectionNameType& PotentialChildCollection : DragDropOp->Collections)
 		{
+			FText Error;
 			bIsValidDrag = CollectionManagerModule.Get().IsValidParentCollection(
 				PotentialChildCollection.Name, PotentialChildCollection.Type,
-				CollectionItem->CollectionName, CollectionItem->CollectionType
-				);
+				CollectionItem->CollectionName, CollectionItem->CollectionType, &Error);
 
 			if (!bIsValidDrag)
 			{
-				DragDropOp->SetToolTip(CollectionManagerModule.Get().GetLastError(), FAppStyle::GetBrush(TEXT("Graph.ConnectorFeedback.Error")));
+				DragDropOp->SetToolTip(Error, FAppStyle::GetBrush(TEXT("Graph.ConnectorFeedback.Error")));
 				break;
 			}
 		}
@@ -1321,12 +1323,13 @@ FReply SCollectionView::HandleDragDropOnCollectionItem(TSharedRef<FCollectionIte
 		// Reparent all of the collections in the drag drop so that they are our immediate children
 		for (const FCollectionNameType& NewChildCollection : DragDropOp->Collections)
 		{
+			FText Error;
 			if (!CollectionManagerModule.Get().ReparentCollection(
 					NewChildCollection.Name, NewChildCollection.Type,
-					CollectionItem->CollectionName, CollectionItem->CollectionType
+					CollectionItem->CollectionName, CollectionItem->CollectionType, &Error
 					))
 			{
-				ContentBrowserUtils::DisplayMessage(CollectionManagerModule.Get().GetLastError(), Geometry.GetLayoutBoundingRect(), SharedThis(this));
+				ContentBrowserUtils::DisplayMessage(Error, Geometry.GetLayoutBoundingRect(), SharedThis(this));
 			}
 		}
 
@@ -1349,7 +1352,7 @@ FReply SCollectionView::HandleDragDropOnCollectionItem(TSharedRef<FCollectionIte
 		int32 NumAdded = 0;
 		FText Message;
 		if (CollectionManagerModule.Get().AddToCollection(
-				CollectionItem->CollectionName, CollectionItem->CollectionType, ObjectPaths, &NumAdded))
+				CollectionItem->CollectionName, CollectionItem->CollectionType, ObjectPaths, &NumAdded, &Message))
 		{
 			if (DroppedAssets.Num() == 1)
 			{
@@ -1377,10 +1380,6 @@ FReply SCollectionView::HandleDragDropOnCollectionItem(TSharedRef<FCollectionIte
 				FTelemetryRouter::Get().ProvideTelemetry(AssetAdded);
 			}
 		}
-		else
-		{
-			Message = CollectionManagerModule.Get().GetLastError();
-		}
 
 		// Added items to the collection or failed. Either way, display the message.
 		ContentBrowserUtils::DisplayMessage(Message, Geometry.GetLayoutBoundingRect(), SharedThis(this));
@@ -1514,13 +1513,14 @@ bool SCollectionView::CollectionNameChangeCommit( const TSharedPtr< FCollectionI
 			return false;
 		}
 
-		if ( !CollectionManagerModule.Get().CreateCollection(NewNameFinal, CollectionItem->CollectionType, CollectionItem->StorageMode) )
+		FText Error;
+		if ( !CollectionManagerModule.Get().CreateCollection(NewNameFinal, CollectionItem->CollectionType, CollectionItem->StorageMode, &Error) )
 		{
 			// Failed to add the collection, remove it from the list
 			AvailableCollections.Remove(FCollectionNameType(CollectionItem->CollectionName, CollectionItem->CollectionType));
 			UpdateFilteredCollectionItems();
 
-			OutWarningMessage = FText::Format( LOCTEXT("CreateCollectionFailed", "Failed to create the collection. {0}"), CollectionManagerModule.Get().GetLastError());
+			OutWarningMessage = FText::Format( LOCTEXT("CreateCollectionFailed", "Failed to create the collection. {0}"), Error);
 			return false;
 		}
 
@@ -1566,10 +1566,11 @@ bool SCollectionView::CollectionNameChangeCommit( const TSharedPtr< FCollectionI
 		}
 
 		// Otherwise perform the rename
-		if ( !CollectionManagerModule.Get().RenameCollection(CollectionItem->CollectionName, CollectionItem->CollectionType, NewNameFinal, CollectionItem->CollectionType) )
+		FText Error;
+		if ( !CollectionManagerModule.Get().RenameCollection(CollectionItem->CollectionName, CollectionItem->CollectionType, NewNameFinal, CollectionItem->CollectionType, &Error) )
 		{
 			// Failed to rename the collection
-			OutWarningMessage = FText::Format( LOCTEXT("RenameCollectionFailed", "Failed to rename the collection. {0}"), CollectionManagerModule.Get().GetLastError());
+			OutWarningMessage = FText::Format( LOCTEXT("RenameCollectionFailed", "Failed to rename the collection. {0}"), Error);
 			return false;
 		}
 	}
@@ -1606,9 +1607,8 @@ bool SCollectionView::CollectionVerifyRenameCommit(const TSharedPtr< FCollection
 
 	FCollectionManagerModule& CollectionManagerModule = FCollectionManagerModule::GetModule();
 
-	if (!CollectionManagerModule.Get().IsValidCollectionName(NewName, ECollectionShareType::CST_All))
+	if (!CollectionManagerModule.Get().IsValidCollectionName(NewName, ECollectionShareType::CST_All, &OutErrorMessage))
 	{
-		OutErrorMessage = CollectionManagerModule.Get().GetLastError();
 		return false;
 	}
 
