Fix missing early out on non-existent collection.

diff --git a/Engine/Source/Developer/CollectionManager/Private/CollectionManager.cpp b/Engine/Source/Developer/CollectionManager/Private/CollectionManager.cpp
--- a/Engine/Source/Developer/CollectionManager/Private/CollectionManager.cpp
+++ b/Engine/Source/Developer/CollectionManager/Private/CollectionManager.cpp
@@ -1847,6 +1847,7 @@ bool FCollectionManager::SetCollectionColor(FName CollectionName, ECollectionSha
 			{
 				*OutError = LOCTEXT("Error_DoesntExist", "The collection doesn't exist.");
 			}
+			return false;
 		}
 
 		(*CollectionRefPtr)->SetCollectionColor(NewColor);
