Attempt to fix PVS-studio error by passing pointer explicitly rather than reference to uninitialized member.

diff --git a/Engine/Source/Developer/CollectionManager/Private/CollectionManager.cpp b/Engine/Source/Developer/CollectionManager/Private/CollectionManager.cpp
--- a/Engine/Source/Developer/CollectionManager/Private/CollectionManager.cpp
+++ b/Engine/Source/Developer/CollectionManager/Private/CollectionManager.cpp
@@ -111,7 +111,7 @@ public:
 class FCollectionManagerCache
 {
 public:
-	FCollectionManagerCache(TMap<FCollectionNameType, TSharedRef<FCollection>>& InAvailableCollections);
+	FCollectionManagerCache(TMap<FCollectionNameType, TSharedRef<FCollection>>* InAvailableCollections);
 
 	/** 
 	 * Dirty the parts of the cache that need to change when a collection is added to our collection manager.
@@ -183,7 +183,7 @@ public:
 
 private:
 	/** Reference to the collections that are currently available in our owner collection manager */
-	TMap<FCollectionNameType, TSharedRef<FCollection>>& AvailableCollections;
+	TMap<FCollectionNameType, TSharedRef<FCollection>>* AvailableCollections;
 
 	/** A map of collection GUIDs to their associated collection names */
 	TMap<FGuid, FCollectionNameType> CachedCollectionNamesFromGuids_Internal;
@@ -204,7 +204,7 @@ private:
 	ERecursiveWorkerFlowControl RecursionHelper_DoWorkOnChildren(FCollectionLock&, const FCollectionNameType& InCollectionKey, FRecursiveWorkerFunc InWorkerFunc) const;
 };
 
-FCollectionManagerCache::FCollectionManagerCache(TMap<FCollectionNameType, TSharedRef<FCollection>>& InAvailableCollections)
+FCollectionManagerCache::FCollectionManagerCache(TMap<FCollectionNameType, TSharedRef<FCollection>>* InAvailableCollections)
 	: AvailableCollections(InAvailableCollections)
 {
 }
@@ -257,7 +257,7 @@ void FCollectionManagerCache::UpdateCaches(FCollectionLock_RW& InGuard, ECollect
 	{
 		CachedCollectionNamesFromGuids_Internal.Reset();
 		EnumRemoveFlags(DirtyFlags, ECollectionCacheFlags::Names);
-		for (const TPair<FCollectionNameType, TSharedRef<FCollection>>& AvailableCollection : AvailableCollections)
+		for (const TPair<FCollectionNameType, TSharedRef<FCollection>>& AvailableCollection : *AvailableCollections)
 		{
 			const FCollectionNameType& CollectionKey = AvailableCollection.Key;
 			const TSharedRef<FCollection>& Collection = AvailableCollection.Value;
@@ -272,7 +272,7 @@ void FCollectionManagerCache::UpdateCaches(FCollectionLock_RW& InGuard, ECollect
 		EnumRemoveFlags(DirtyFlags, ECollectionCacheFlags::Hierarchy);
 		const TMap<FGuid, FCollectionNameType>& CachedCollectionNamesFromGuids = GetCachedCollectionNamesFromGuids(InGuard);
 
-		for (const TPair<FCollectionNameType, TSharedRef<FCollection>>& AvailableCollection : AvailableCollections)
+		for (const TPair<FCollectionNameType, TSharedRef<FCollection>>& AvailableCollection : *AvailableCollections)
 		{
 			const FCollectionNameType& CollectionKey = AvailableCollection.Key;
 			const TSharedRef<FCollection>& Collection = AvailableCollection.Value;
@@ -292,7 +292,7 @@ void FCollectionManagerCache::UpdateCaches(FCollectionLock_RW& InGuard, ECollect
 		CachedObjects_Internal.Reset();
 		EnumRemoveFlags(DirtyFlags, ECollectionCacheFlags::Objects);
 
-		for (const TPair<FCollectionNameType, TSharedRef<FCollection>>& AvailableCollection : AvailableCollections)
+		for (const TPair<FCollectionNameType, TSharedRef<FCollection>>& AvailableCollection : *AvailableCollections)
 		{
 			const FCollectionNameType& CollectionKey = AvailableCollection.Key;
 			const TSharedRef<FCollection>& Collection = AvailableCollection.Value;
@@ -349,7 +349,7 @@ void FCollectionManagerCache::UpdateCaches(FCollectionLock_RW& InGuard, ECollect
 	{
 		CachedColors_Internal.Reset();
 		EnumRemoveFlags(DirtyFlags, ECollectionCacheFlags::Colors);
-		for (const TPair<FCollectionNameType, TSharedRef<FCollection>>& AvailableCollection : AvailableCollections)
+		for (const TPair<FCollectionNameType, TSharedRef<FCollection>>& AvailableCollection : *AvailableCollections)
 		{
 			const TSharedRef<FCollection>& Collection = AvailableCollection.Value;
 
@@ -363,7 +363,7 @@ void FCollectionManagerCache::UpdateCaches(FCollectionLock_RW& InGuard, ECollect
 		}
 	}
 
-	UE_LOG(LogCollectionManager, Verbose, TEXT("Rebuilt caches for %d collections in in %0.6f seconds"), AvailableCollections.Num(), FPlatformTime::Seconds() - CacheStartTime);
+	UE_LOG(LogCollectionManager, Verbose, TEXT("Rebuilt caches for %d collections in in %0.6f seconds"), AvailableCollections->Num(), FPlatformTime::Seconds() - CacheStartTime);
 }
 
 void FCollectionManagerCache::RecursionHelper_DoWork(FCollectionLock& Guard,
@@ -392,7 +392,7 @@ void FCollectionManagerCache::RecursionHelper_DoWork(FCollectionLock& Guard,
 FCollectionManagerCache::ERecursiveWorkerFlowControl FCollectionManagerCache::RecursionHelper_DoWorkOnParents(
 	FCollectionLock& Guard, const FCollectionNameType& InCollectionKey, FRecursiveWorkerFunc InWorkerFunc) const
 {
-	const TSharedRef<FCollection>* const CollectionRefPtr = AvailableCollections.Find(InCollectionKey);
+	const TSharedRef<FCollection>* const CollectionRefPtr = AvailableCollections->Find(InCollectionKey);
 	if (CollectionRefPtr)
 	{
 		const TMap<FGuid, FCollectionNameType>& CachedCollectionNamesFromGuids = GetCachedCollectionNamesFromGuids(Guard);
@@ -413,7 +413,7 @@ FCollectionManagerCache::ERecursiveWorkerFlowControl FCollectionManagerCache::Re
 FCollectionManagerCache::ERecursiveWorkerFlowControl FCollectionManagerCache::RecursionHelper_DoWorkOnChildren(
 	FCollectionLock& Guard, const FCollectionNameType& InCollectionKey, FRecursiveWorkerFunc InWorkerFunc) const
 {
-	const TSharedRef<FCollection>* const CollectionRefPtr = AvailableCollections.Find(InCollectionKey);
+	const TSharedRef<FCollection>* const CollectionRefPtr = AvailableCollections->Find(InCollectionKey);
 	if (CollectionRefPtr)
 	{
 		const TMap<FGuid, TArray<FGuid>>& CachedHierarchy = GetCachedHierarchy(Guard);
@@ -467,7 +467,7 @@ const TArray<FLinearColor>& FCollectionManagerCache::GetCachedColors(FCollection
 FStringView FCollectionManager::CollectionExtension = TEXTVIEW("collection");
 
 FCollectionManager::FCollectionManager()
-	: CollectionCache(MakePimpl<FCollectionManagerCache>(AvailableCollections))
+	: CollectionCache(MakePimpl<FCollectionManagerCache>(&AvailableCollections))
 {
 	CollectionFolders[ECollectionShareType::CST_Local] = FPaths::ProjectSavedDir() / TEXT("Collections");
 	CollectionFolders[ECollectionShareType::CST_Private] = FPaths::GameUserDeveloperDir() / TEXT("Collections");
