Fix unnecessary dereference of OutError exposed by PVS studio

diff --git a/Engine/Source/Developer/CollectionManager/Private/CollectionManager.cpp b/Engine/Source/Developer/CollectionManager/Private/CollectionManager.cpp
--- a/Engine/Source/Developer/CollectionManager/Private/CollectionManager.cpp
+++ b/Engine/Source/Developer/CollectionManager/Private/CollectionManager.cpp
@@ -975,7 +975,7 @@ bool FCollectionManager::IsValidCollectionName(const FString& CollectionName, EC
 	}
 
 	// Make sure the new name only contains valid characters
-	if (!CollectionNameFinal.IsValidXName(INVALID_OBJECTNAME_CHARACTERS INVALID_LONGPACKAGE_CHARACTERS, &*OutError))
+	if (!CollectionNameFinal.IsValidXName(INVALID_OBJECTNAME_CHARACTERS INVALID_LONGPACKAGE_CHARACTERS, OutError))
 	{
 		return false;
 	}
