Collections should respect auto-commit on submit setting when referencing deleted redirectors
* Redirector deletion is not auto-submitted, this allows users to submit both collection change and redirector deletion change as part of the same CL

diff --git a/Engine/Source/Developer/CollectionManager/Private/CollectionManager.cpp b/Engine/Source/Developer/CollectionManager/Private/CollectionManager.cpp
--- a/Engine/Source/Developer/CollectionManager/Private/CollectionManager.cpp
+++ b/Engine/Source/Developer/CollectionManager/Private/CollectionManager.cpp
@@ -2102,7 +2102,7 @@ bool FCollectionManager::HandleRedirectorsDeleted(TConstArrayView<FSoftObjectPat
 				const TSharedRef<FCollection>& Collection = *CollectionRefPtr;
 
 				FText SaveError;
-				constexpr bool bForceCommitToRevisionControl = true;
+				constexpr bool bForceCommitToRevisionControl = false;
 				if (InternalSaveCollection(Guard, Collection, &SaveError, bForceCommitToRevisionControl))
 				{
 					CollectionFileCaches[CollectionKey.Type]->IgnoreFileModification(Collection->GetSourceFilename());
