Removed a number of functions deprecated in 5.1 to do with using FNames for full object paths instead of FSoftObjectPaths in an attempt to clean up the new ICollectionManager interface which has a number of additional deprecated functions.

diff --git a/Engine/Plugins/Editor/AssetManagerEditor/Source/AssetManagerEditor/Public/AssetManagerEditorModule.h b/Engine/Plugins/Editor/AssetManagerEditor/Source/AssetManagerEditor/Public/AssetManagerEditorModule.h
--- a/Engine/Plugins/Editor/AssetManagerEditor/Source/AssetManagerEditor/Public/AssetManagerEditorModule.h
+++ b/Engine/Plugins/Editor/AssetManagerEditor/Source/AssetManagerEditor/Public/AssetManagerEditorModule.h
@@ -236,14 +236,6 @@ public:
 		return false;
 	}
 
-	UE_DEPRECATED(5.1, "Asset path FNames have been deprecated, use FSoftObjectPath instead.")
-	FAssetData GetAssetByObjectPath(FName ObjectPath) const
-	{
-PRAGMA_DISABLE_DEPRECATION_WARNINGS
-		return GetAssetByObjectPath(FSoftObjectPath(ObjectPath));
-PRAGMA_ENABLE_DEPRECATION_WARNINGS
-	}
-
 	bool GetDependencies(const FAssetIdentifier& AssetIdentifier, TArray<FAssetIdentifier>& OutDependencies,
 		UE::AssetRegistry::EDependencyCategory Category = UE::AssetRegistry::EDependencyCategory::All,
 		const UE::AssetRegistry::FDependencyQuery& Flags = UE::AssetRegistry::FDependencyQuery()) const
diff --git a/Engine/Source/Developer/CollectionManager/Private/CollectionManager.h b/Engine/Source/Developer/CollectionManager/Private/CollectionManager.h
--- a/Engine/Source/Developer/CollectionManager/Private/CollectionManager.h
+++ b/Engine/Source/Developer/CollectionManager/Private/CollectionManager.h
@@ -66,99 +66,6 @@ public:
 	FCollectionManager();
 	virtual ~FCollectionManager();
 
-	// ICollectionManager implementation
-PRAGMA_DISABLE_DEPRECATION_WARNINGS
-	virtual bool GetAssetsInCollection(FName CollectionName, ECollectionShareType::Type ShareType, TArray<FName>& AssetPaths, ECollectionRecursionFlags::Flags RecursionMode = ECollectionRecursionFlags::Self) const override
-	{
-		TArray<FSoftObjectPath> Temp;
-		if (GetAssetsInCollection(CollectionName, ShareType, Temp, RecursionMode))
-		{
-			AssetPaths.Append(UE::SoftObjectPath::Private::ConvertSoftObjectPaths(Temp));
-			return true;
-		}
-		return false;
-	}
-	virtual bool GetObjectsInCollection(FName CollectionName, ECollectionShareType::Type ShareType, TArray<FName>& ObjectPaths, ECollectionRecursionFlags::Flags RecursionMode = ECollectionRecursionFlags::Self) const override
-	{
-		TArray<FSoftObjectPath> Temp;
-		if (GetObjectsInCollection(CollectionName, ShareType, Temp, RecursionMode))
-		{
-			ObjectPaths.Append(UE::SoftObjectPath::Private::ConvertSoftObjectPaths(Temp));
-			return true;
-		}
-		return false;
-	}
-	virtual bool GetClassesInCollection(FName CollectionName, ECollectionShareType::Type ShareType, TArray<FName>& ClassPaths, ECollectionRecursionFlags::Flags RecursionMode = ECollectionRecursionFlags::Self) const override
-	{
-		TArray<FTopLevelAssetPath> Temp;
-		if (GetClassesInCollection(CollectionName, ShareType, Temp, RecursionMode))
-		{
-			for (FTopLevelAssetPath Path : Temp)
-			{
-				ClassPaths.Add(Path.ToFName());
-			}
-		
-			return true;
-		}
-		return false;
-	}
-	virtual void GetCollectionsContainingObject(FName ObjectPath, ECollectionShareType::Type ShareType, TArray<FName>& OutCollectionNames, ECollectionRecursionFlags::Flags RecursionMode = ECollectionRecursionFlags::Self) const override
-	{
-		GetCollectionsContainingObject(FSoftObjectPath(ObjectPath), ShareType, OutCollectionNames, RecursionMode);
-	}
-	virtual void GetCollectionsContainingObject(FName ObjectPath, TArray<FCollectionNameType>& OutCollections, ECollectionRecursionFlags::Flags RecursionMode = ECollectionRecursionFlags::Self) const override
-	{
-		GetCollectionsContainingObject(FSoftObjectPath(ObjectPath), OutCollections, RecursionMode);
-	}
-	virtual void GetCollectionsContainingObjects(const TArray<FName>& ObjectPathNames, TMap<FCollectionNameType, TArray<FName>>& OutCollectionsAndMatchedObjects, ECollectionRecursionFlags::Flags RecursionMode = ECollectionRecursionFlags::Self) const override
-	{
-		TArray<FSoftObjectPath> Paths = UE::SoftObjectPath::Private::ConvertObjectPathNames(ObjectPathNames);
-		TMap<FCollectionNameType, TArray<FSoftObjectPath>> TmpMap;
-		GetCollectionsContainingObjects(Paths, TmpMap, RecursionMode);
-		for (const TPair<FCollectionNameType, TArray<FSoftObjectPath>>& Pair : TmpMap)
-		{
-			TArray<FName>& Names = OutCollectionsAndMatchedObjects.FindOrAdd(Pair.Key);
-			Names.Append(UE::SoftObjectPath::Private::ConvertSoftObjectPaths(Pair.Value));
-		}
-	}
-	virtual FString GetCollectionsStringForObject(FName ObjectPath, ECollectionShareType::Type ShareType, ECollectionRecursionFlags::Flags RecursionMode = ECollectionRecursionFlags::Self, bool bFullPaths = true) const override
-	{
-		return GetCollectionsStringForObject(FSoftObjectPath(ObjectPath), ShareType, RecursionMode, bFullPaths);
-	}
-	virtual bool AddToCollection(FName CollectionName, ECollectionShareType::Type ShareType, FName ObjectPath) override
-	{
-		return AddToCollection(CollectionName, ShareType, FSoftObjectPath(ObjectPath));
-	}
-	virtual bool AddToCollection(FName CollectionName, ECollectionShareType::Type ShareType, const TArray<FName>& ObjectPaths, int32* OutNumAdded = nullptr) override
-	{
-		return AddToCollection(CollectionName, ShareType, UE::SoftObjectPath::Private::ConvertObjectPathNames(ObjectPaths), OutNumAdded);
-	}
-	virtual bool RemoveFromCollection(FName CollectionName, ECollectionShareType::Type ShareType, FName ObjectPath) override
-	{
-		return RemoveFromCollection(CollectionName, ShareType, FSoftObjectPath(ObjectPath));
-	}
-	virtual bool RemoveFromCollection(FName CollectionName, ECollectionShareType::Type ShareType, const TArray<FName>& ObjectPaths, int32* OutNumRemoved = nullptr) override
-	{
-		return RemoveFromCollection(CollectionName, ShareType, UE::SoftObjectPath::Private::ConvertObjectPathNames(ObjectPaths), OutNumRemoved);
-	}
-	virtual bool IsObjectInCollection(FName ObjectPath, FName CollectionName, ECollectionShareType::Type ShareType, ECollectionRecursionFlags::Flags RecursionMode = ECollectionRecursionFlags::Self) const override
-	{
-		return IsObjectInCollection(FSoftObjectPath(ObjectPath), CollectionName, ShareType, RecursionMode);
-	}
-	virtual bool HandleRedirectorDeleted(const FName& ObjectPath) override
-	{
-		return HandleRedirectorDeleted(FSoftObjectPath(ObjectPath));
-	}
-	virtual void HandleObjectRenamed(const FName& OldObjectPath, const FName& NewObjectPath) override
-	{
-		return HandleObjectRenamed(FSoftObjectPath(OldObjectPath), FSoftObjectPath(NewObjectPath));
-	}
-	virtual void HandleObjectDeleted(const FName& ObjectPath) override
-	{
-		return HandleObjectDeleted(FSoftObjectPath(ObjectPath));
-	}
-	PRAGMA_ENABLE_DEPRECATION_WARNINGS
-
 	virtual const TSharedRef<ICollectionContainer>& GetProjectCollectionContainer() const override;
 	virtual TSharedPtr<ICollectionContainer> AddCollectionContainer(const TSharedRef<ICollectionSource>& CollectionSource) override;
 	virtual bool RemoveCollectionContainer(const TSharedRef<ICollectionContainer>& CollectionContainer) override;
diff --git a/Engine/Source/Developer/CollectionManager/Public/ICollectionManager.h b/Engine/Source/Developer/CollectionManager/Public/ICollectionManager.h
--- a/Engine/Source/Developer/CollectionManager/Public/ICollectionManager.h
+++ b/Engine/Source/Developer/CollectionManager/Public/ICollectionManager.h
@@ -92,8 +92,6 @@ public:
 	 */
 	UE_DEPRECATED(5.6, "ICollectionContainer has been introduced as a new top level object for housing collections, call GetProjectCollectionContainer to disambiguate which collection container to use.")
 	virtual bool GetAssetsInCollection(FName CollectionName, ECollectionShareType::Type ShareType, TArray<FSoftObjectPath>& OutAssetPaths, ECollectionRecursionFlags::Flags RecursionMode = ECollectionRecursionFlags::Self) const = 0;
-	UE_DEPRECATED(5.1, "FNames containing full object paths are deprecated. Use FSoftObjectPath instead.")
-	virtual bool GetAssetsInCollection(FName CollectionName, ECollectionShareType::Type ShareType, TArray<FName>& OutAssetPaths, ECollectionRecursionFlags::Flags RecursionMode = ECollectionRecursionFlags::Self) const = 0;
 
 	/**
 	 * Gets a list of class paths in the given collection
@@ -102,8 +100,6 @@ public:
 	 */
 	UE_DEPRECATED(5.6, "ICollectionContainer has been introduced as a new top level object for housing collections, call GetProjectCollectionContainer to disambiguate which collection container to use.")
 	virtual bool GetClassesInCollection(FName CollectionName, ECollectionShareType::Type ShareType, TArray<FTopLevelAssetPath>& ClassPaths, ECollectionRecursionFlags::Flags RecursionMode = ECollectionRecursionFlags::Self) const = 0;
-	UE_DEPRECATED(5.1, "FNames containing class paths are deprecated. Use FTopLevelAssetPath instead.")
-	virtual bool GetClassesInCollection(FName CollectionName, ECollectionShareType::Type ShareType, TArray<FName>& ClassPaths, ECollectionRecursionFlags::Flags RecursionMode = ECollectionRecursionFlags::Self) const = 0;
 
 	/**
 	 * Gets a list of objects in the given collection
@@ -112,32 +108,22 @@ public:
 	 */
 	UE_DEPRECATED(5.6, "ICollectionContainer has been introduced as a new top level object for housing collections, call GetProjectCollectionContainer to disambiguate which collection container to use.")
 	virtual bool GetObjectsInCollection(FName CollectionName, ECollectionShareType::Type ShareType, TArray<FSoftObjectPath>& ObjectPaths, ECollectionRecursionFlags::Flags RecursionMode = ECollectionRecursionFlags::Self) const = 0;
-	UE_DEPRECATED(5.1, "FNames containing full object paths are deprecated. Use FSoftObjectPath instead.")
-	virtual bool GetObjectsInCollection(FName CollectionName, ECollectionShareType::Type ShareType, TArray<FName>& ObjectPaths, ECollectionRecursionFlags::Flags RecursionMode = ECollectionRecursionFlags::Self) const = 0;
 
 	/** Returns a list of collections in which the specified object exists of the specified share type */
 	UE_DEPRECATED(5.6, "ICollectionContainer has been introduced as a new top level object for housing collections, call GetProjectCollectionContainer to disambiguate which collection container to use.")
 	virtual void GetCollectionsContainingObject(const FSoftObjectPath& ObjectPath, ECollectionShareType::Type ShareType, TArray<FName>& OutCollectionNames, ECollectionRecursionFlags::Flags RecursionMode = ECollectionRecursionFlags::Self) const = 0;
-	UE_DEPRECATED(5.1, "FNames containing full object paths are deprecated. Use FSoftObjectPath instead.")
-	virtual void GetCollectionsContainingObject(FName ObjectPath, ECollectionShareType::Type ShareType, TArray<FName>& OutCollectionNames, ECollectionRecursionFlags::Flags RecursionMode = ECollectionRecursionFlags::Self) const = 0;
 
 	/** Returns a list of collections in which the specified object exists */
 	UE_DEPRECATED(5.6, "ICollectionContainer has been introduced as a new top level object for housing collections, call GetProjectCollectionContainer to disambiguate which collection container to use.")
 	virtual void GetCollectionsContainingObject(const FSoftObjectPath& ObjectPath, TArray<FCollectionNameType>& OutCollections, ECollectionRecursionFlags::Flags RecursionMode = ECollectionRecursionFlags::Self) const = 0;
-	UE_DEPRECATED(5.1, "FNames containing full object paths are deprecated. Use FSoftObjectPath instead.")
-	virtual void GetCollectionsContainingObject(FName ObjectPath, TArray<FCollectionNameType>& OutCollections, ECollectionRecursionFlags::Flags RecursionMode = ECollectionRecursionFlags::Self) const = 0;
 
 	/** Returns a list of collections in which any of the specified objects exist */
 	UE_DEPRECATED(5.6, "ICollectionContainer has been introduced as a new top level object for housing collections, call GetProjectCollectionContainer to disambiguate which collection container to use.")
 	virtual void GetCollectionsContainingObjects(const TArray<FSoftObjectPath>& ObjectPaths, TMap<FCollectionNameType, TArray<FSoftObjectPath>>& OutCollectionsAndMatchedObjects, ECollectionRecursionFlags::Flags RecursionMode = ECollectionRecursionFlags::Self) const = 0;
-	UE_DEPRECATED(5.1, "FNames containing full object paths are deprecated. Use FSoftObjectPath instead.")
-	virtual void GetCollectionsContainingObjects(const TArray<FName>& ObjectPaths, TMap<FCollectionNameType, TArray<FName>>& OutCollectionsAndMatchedObjects, ECollectionRecursionFlags::Flags RecursionMode = ECollectionRecursionFlags::Self) const = 0;
 
 	/** Returns a string containing a comma separated list of collections in which the specified object exists of the specified share type */
 	UE_DEPRECATED(5.6, "ICollectionContainer has been introduced as a new top level object for housing collections, call GetProjectCollectionContainer to disambiguate which collection container to use.")
 	virtual FString GetCollectionsStringForObject(const FSoftObjectPath& ObjectPath, ECollectionShareType::Type ShareType, ECollectionRecursionFlags::Flags RecursionMode = ECollectionRecursionFlags::Self, bool bFullPaths = true) const = 0;
-	UE_DEPRECATED(5.1, "FNames containing full object paths are deprecated. Use FSoftObjectPath instead.")
-	virtual FString GetCollectionsStringForObject(FName ObjectPath, ECollectionShareType::Type ShareType, ECollectionRecursionFlags::Flags RecursionMode = ECollectionRecursionFlags::Self, bool bFullPaths = true) const = 0;
 
 	/** 
 	 * Creates a unique collection name for the given type taking the form BaseName+(unique number) 
@@ -235,11 +221,6 @@ public:
 		FName CollectionName, ECollectionShareType::Type ShareType, TConstArrayView<FSoftObjectPath> ObjectPaths,
 		int32* OutNumAdded = nullptr, FText* OutError = nullptr) = 0;
 
-	UE_DEPRECATED(5.1, "FNames containing full object paths are deprecated. Use FSoftObjectPath instead.")
-	virtual bool AddToCollection(FName CollectionName, ECollectionShareType::Type ShareType, FName ObjectPath) = 0;
-	UE_DEPRECATED(5.1, "FNames containing full object paths are deprecated. Use FSoftObjectPath instead.")
-	virtual bool AddToCollection(FName CollectionName, ECollectionShareType::Type ShareType, const TArray<FName>& ObjectPaths, int32* OutNumAdded = nullptr) = 0;
-
 	/**
 	 * Removes the asset from the specified collection.
 	 *
@@ -258,10 +239,6 @@ public:
 	virtual bool RemoveFromCollection(
 		FName CollectionName, ECollectionShareType::Type ShareType, TConstArrayView<FSoftObjectPath> ObjectPaths,
 		int32* OutNumRemoved = nullptr, FText* OutError = nullptr) = 0;
-	UE_DEPRECATED(5.1, "FNames containing full object paths are deprecated. Use FSoftObjectPath instead.")
-	virtual bool RemoveFromCollection(FName CollectionName, ECollectionShareType::Type ShareType, FName ObjectPath) = 0;
-	UE_DEPRECATED(5.1, "FNames containing full object paths are deprecated. Use FSoftObjectPath instead.")
-	virtual bool RemoveFromCollection(FName CollectionName, ECollectionShareType::Type ShareType, const TArray<FName>& ObjectPaths, int32* OutNumRemoved = nullptr) = 0;
 
 	/**
 	 * Sets the dynamic query text for the specified collection. 
@@ -427,8 +404,6 @@ public:
 		const FSoftObjectPath& ObjectPath, FName CollectionName, ECollectionShareType::Type ShareType,
 		ECollectionRecursionFlags::Flags RecursionMode = ECollectionRecursionFlags::Self,
 		FText* OutError = nullptr) const = 0;
-	UE_DEPRECATED(5.1, "FNames containing full object paths are deprecated. Use FSoftObjectPath instead.")
-	virtual bool IsObjectInCollection(FName ObjectPath, FName CollectionName, ECollectionShareType::Type ShareType, ECollectionRecursionFlags::Flags RecursionMode = ECollectionRecursionFlags::Self) const = 0;
 
 	/**
 	 * Check to see if the given collection is valid to be used as the parent of another collection.
@@ -461,8 +436,6 @@ public:
 	 * @return true if all of the collections that were referencing this redirector could be re-saved, false otherwise
 	 */
 	virtual bool HandleRedirectorDeleted(const FSoftObjectPath& ObjectPath, FText* OutError = nullptr) = 0;
-	UE_DEPRECATED(5.1, "FNames containing full object paths are deprecated. Use FSoftObjectPath instead.")
-	virtual bool HandleRedirectorDeleted(const FName& ObjectPath) = 0;
 
 	/**
 	 * Called to notify the collections that redirectors have been deleted and that they should ensure their on-disk representation is re-saved with the fixed up in-memory version
@@ -472,13 +445,9 @@ public:
 
 	/** Called to notify the collections that an object has been renamed or moved */
 	virtual void HandleObjectRenamed(const FSoftObjectPath& OldObjectPath, const FSoftObjectPath& NewObjectPath) = 0;
-	UE_DEPRECATED(5.1, "FNames containing full object paths are deprecated. Use FSoftObjectPath instead.")
-	virtual void HandleObjectRenamed(const FName& OldObjectPath, const FName& NewObjectPath) = 0;
 
 	/** Called to notify the collections that an object has been deleted */
 	virtual void HandleObjectDeleted(const FSoftObjectPath& ObjectPath) = 0;
-	UE_DEPRECATED(5.1, "FNames containing full object paths are deprecated. Use FSoftObjectPath instead.")
-	virtual void HandleObjectDeleted(const FName& ObjectPath) = 0;
 
 	/** Called to notify the collections that objects have been deleted */
 	virtual void HandleObjectsDeleted(TConstArrayView<FSoftObjectPath> ObjectPaths) = 0;
diff --git a/Engine/Source/Editor/ContentBrowserData/Private/ContentBrowserItem.cpp b/Engine/Source/Editor/ContentBrowserData/Private/ContentBrowserItem.cpp
--- a/Engine/Source/Editor/ContentBrowserData/Private/ContentBrowserItem.cpp
+++ b/Engine/Source/Editor/ContentBrowserData/Private/ContentBrowserItem.cpp
@@ -606,13 +606,6 @@ bool FContentBrowserItem::TryGetCollectionId(FSoftObjectPath& OutCollectionId) c
 	return FContentBrowserItemHelper::CallDataSourceImpl<UContentBrowserDataSource>(*this, UE_PROJECTION_MEMBER(UContentBrowserDataSource, TryGetCollectionId), OutCollectionId);
 }
 
-bool FContentBrowserItem::TryGetCollectionId(FName& OutCollectionId) const
-{
-PRAGMA_DISABLE_DEPRECATION_WARNINGS
-	return FContentBrowserItemHelper::CallDataSourceImpl<UContentBrowserDataSource>(*this, UE_PROJECTION_MEMBER(UContentBrowserDataSource, TryGetCollectionId), OutCollectionId);
-PRAGMA_ENABLE_DEPRECATION_WARNINGS
-}
-
 bool FContentBrowserItem::Legacy_TryGetPackagePath(FName& OutPackagePath) const
 {
 	return FContentBrowserItemHelper::CallDataSourceImpl<UContentBrowserDataSource>(*this, UE_PROJECTION_MEMBER(UContentBrowserDataSource, Legacy_TryGetPackagePath), OutPackagePath);
diff --git a/Engine/Source/Editor/ContentBrowserData/Public/ContentBrowserDataSource.h b/Engine/Source/Editor/ContentBrowserData/Public/ContentBrowserDataSource.h
--- a/Engine/Source/Editor/ContentBrowserData/Public/ContentBrowserDataSource.h
+++ b/Engine/Source/Editor/ContentBrowserData/Public/ContentBrowserDataSource.h
@@ -750,20 +750,6 @@ public:
 	 */
 	virtual bool TryGetCollectionId(const FContentBrowserItemData& InItem, FSoftObjectPath& OutCollectionId);
 
-	UE_DEPRECATED(5.1, "FNames containing full object paths are deprecated. Use FSoftObjectPath instead.")
-	bool TryGetCollectionId(const FContentBrowserItemData& InItem, FName& OutCollectionId)
-	{
-PRAGMA_DISABLE_DEPRECATION_WARNINGS
-		FSoftObjectPath Temp;
-		if (TryGetCollectionId(InItem, Temp))
-		{
-			OutCollectionId = Temp.ToFName();
-			return true;
-		}
-		return false;
-PRAGMA_ENABLE_DEPRECATION_WARNINGS
-	}
-
 	/**
 	 * Attempt to retrieve the package path associated with the given item.
 	 * @note This exists to allow the Content Browser to interface with external callbacks that only operate on package paths and should ideally be avoided for new code.
diff --git a/Engine/Source/Editor/ContentBrowserData/Public/ContentBrowserItem.h b/Engine/Source/Editor/ContentBrowserData/Public/ContentBrowserItem.h
--- a/Engine/Source/Editor/ContentBrowserData/Public/ContentBrowserItem.h
+++ b/Engine/Source/Editor/ContentBrowserData/Public/ContentBrowserItem.h
@@ -445,9 +445,6 @@ public:
 	 */
 	bool TryGetCollectionId(FSoftObjectPath& OutCollectionId) const;
 
-	UE_DEPRECATED(5.1, "FNames containing full object paths are deprecated. Use FSoftObjectPath instead.")
-	bool TryGetCollectionId(FName & OutCollectionId) const;
-
 	/**
 	 * Attempt to retrieve the package path associated with this item.
 	 * @note This exists to allow the Content Browser to interface with external callbacks that only operate on package paths and should ideally be avoided for new code.
diff --git a/Engine/Source/Editor/UnrealEd/Private/Commandlets/PackageUtilities.cpp b/Engine/Source/Editor/UnrealEd/Private/Commandlets/PackageUtilities.cpp
--- a/Engine/Source/Editor/UnrealEd/Private/Commandlets/PackageUtilities.cpp
+++ b/Engine/Source/Editor/UnrealEd/Private/Commandlets/PackageUtilities.cpp
@@ -644,20 +644,6 @@ bool FContentHelper::QueryAssetsInCollection(FName InCollectionName, ECollection
 	return this->QuerySetContent<FCollectionPolicy>(InCollectionName, InType, OutAssetPaths);
 }
 
-bool FContentHelper::QueryAssetsInCollection(FName InCollectionName, ECollectionShareType::Type InType, TArray<FName>& OutAssetPathNames)
-{
-	TArray<FSoftObjectPath> AssetPaths;
-PRAGMA_DISABLE_DEPRECATION_WARNINGS
-	if (this->QuerySetContent<FCollectionPolicy>(InCollectionName, InType, AssetPaths))
-	{
-		OutAssetPathNames.Append(UE::SoftObjectPath::Private::ConvertSoftObjectPaths(AssetPaths));
-		return true;
-	}
-PRAGMA_ENABLE_DEPRECATION_WARNINGS
-
-	return false;
-}
-
 /*-----------------------------------------------------------------------------
 ULoadPackageCommandlet
 -----------------------------------------------------------------------------*/
diff --git a/Engine/Source/Editor/UnrealEd/Public/PackageHelperFunctions.h b/Engine/Source/Editor/UnrealEd/Public/PackageHelperFunctions.h
--- a/Engine/Source/Editor/UnrealEd/Public/PackageHelperFunctions.h
+++ b/Engine/Source/Editor/UnrealEd/Public/PackageHelperFunctions.h
@@ -163,10 +163,6 @@ public:
 	 *	@return True if collection was created successfully
 	 */
 	UNREALED_API bool QueryAssetsInCollection(FName InCollectionName, ECollectionShareType::Type InType, TArray<FSoftObjectPath>& OutAssetPaths);
-	
-	UE_DEPRECATED(5.1, "FNames containing full asset paths are deprecated. Use FSoftObjectPath instead.")
-	UNREALED_API bool QueryAssetsInCollection(FName InCollectionName, ECollectionShareType::Type InType, TArray<FName>& OutAssetPathNames);
-
 
 protected:
 	/** Creates a new collection */
diff --git a/Engine/Source/Runtime/CoreUObject/Private/Misc/RedirectCollector.cpp b/Engine/Source/Runtime/CoreUObject/Private/Misc/RedirectCollector.cpp
--- a/Engine/Source/Runtime/CoreUObject/Private/Misc/RedirectCollector.cpp
+++ b/Engine/Source/Runtime/CoreUObject/Private/Misc/RedirectCollector.cpp
@@ -340,13 +340,6 @@ bool FRedirectCollector::TryRemoveObjectPathRedirectionInternal(const FSoftObjec
 	return bRemoved;
 }
 
-void FRedirectCollector::AddAssetPathRedirection(FName OriginalPath, FName RedirectedPath)
-{
-PRAGMA_DISABLE_DEPRECATION_WARNINGS
-	AddAssetPathRedirection(FSoftObjectPath(OriginalPath), FSoftObjectPath(RedirectedPath));
-PRAGMA_ENABLE_DEPRECATION_WARNINGS
-}
-
 void FRedirectCollector::RemoveAssetPathRedirection(const FSoftObjectPath& OriginalPath)
 {
 	FScopeLock ScopeLock(&CriticalSection);
@@ -354,20 +347,6 @@ void FRedirectCollector::RemoveAssetPathRedirection(const FSoftObjectPath& Origi
 	TryRemoveObjectPathRedirectionInternal(OriginalPath);
 }
 
-void FRedirectCollector::RemoveAssetPathRedirection(FName OriginalPath)
-{
-PRAGMA_DISABLE_DEPRECATION_WARNINGS
-	RemoveAssetPathRedirection(FSoftObjectPath(OriginalPath));
-PRAGMA_ENABLE_DEPRECATION_WARNINGS
-}
-
-FName FRedirectCollector::GetAssetPathRedirection(FName OriginalPath)
-{
-PRAGMA_DISABLE_DEPRECATION_WARNINGS
-	return FName(*GetAssetPathRedirection(FSoftObjectPath(OriginalPath)).ToString());
-PRAGMA_ENABLE_DEPRECATION_WARNINGS
-}
-
 FSoftObjectPath FRedirectCollector::GetAssetPathRedirection(const FSoftObjectPath& OriginalPath)
 {
 	FScopeLock ScopeLock(&CriticalSection);
diff --git a/Engine/Source/Runtime/CoreUObject/Public/Misc/RedirectCollector.h b/Engine/Source/Runtime/CoreUObject/Public/Misc/RedirectCollector.h
--- a/Engine/Source/Runtime/CoreUObject/Public/Misc/RedirectCollector.h
+++ b/Engine/Source/Runtime/CoreUObject/Public/Misc/RedirectCollector.h
@@ -107,18 +107,12 @@ public:
 	COREUOBJECT_API void ProcessSoftObjectPathPackageList(FName FilterPackage, bool bGetEditorOnly, TSet<FName>& OutReferencedPackages);
 
 	/** Adds a new mapping for redirector path to destination path, this is called from the Asset Registry to register all redirects it knows about */
-	UE_DEPRECATED(5.1, "Asset path FNames have been deprecated, use FSoftObjectPath instead.")
-	COREUOBJECT_API void AddAssetPathRedirection(FName OriginalPath, FName RedirectedPath);
 	COREUOBJECT_API void AddAssetPathRedirection(const FSoftObjectPath& OriginalPath, const FSoftObjectPath& RedirectedPath);
 
 	/** Removes an asset path redirection, call this when deleting redirectors */
-	UE_DEPRECATED(5.1, "Asset path FNames have been deprecated, use FSoftObjectPath instead.")
-	COREUOBJECT_API void RemoveAssetPathRedirection(FName OriginalPath);
 	COREUOBJECT_API void RemoveAssetPathRedirection(const FSoftObjectPath& OriginalPath);
 
 	/** Returns a remapped asset path, if it returns null there is no relevant redirector */
-	UE_DEPRECATED(5.1, "Asset path FNames have been deprecated, use FSoftObjectPath instead.")
-	COREUOBJECT_API FName GetAssetPathRedirection(FName OriginalPath);
 	/** Returns a remapped asset path, if there is no relevant redirector, the return value reports true from IsNull() */
 	COREUOBJECT_API FSoftObjectPath GetAssetPathRedirection(const FSoftObjectPath& OriginalPath);
 
