Fixed typo whereby FCollectionContainer::EmptyCollection returns false when it successfully empties a collection rather than true.

diff --git a/Engine/Source/Developer/CollectionManager/Private/CollectionContainer.cpp b/Engine/Source/Developer/CollectionManager/Private/CollectionContainer.cpp
--- a/Engine/Source/Developer/CollectionManager/Private/CollectionContainer.cpp
+++ b/Engine/Source/Developer/CollectionManager/Private/CollectionContainer.cpp
@@ -1665,7 +1665,7 @@ bool FCollectionContainer::EmptyCollection(FName CollectionName, ECollectionShar
 		(*CollectionRefPtr)->Empty();
 		
 		constexpr bool bForceCommitToRevisionControl = true;
-		if (InternalSaveCollection(Guard, *CollectionRefPtr, OutError, bForceCommitToRevisionControl))
+		if (!InternalSaveCollection(Guard, *CollectionRefPtr, OutError, bForceCommitToRevisionControl))
 		{
 			return false;
 		}
