[CollectionManager]
- Optimize loading for big collections
- 1.1s to 541ms at boot for some projects

diff --git a/Engine/Source/Developer/CollectionManager/Private/Collection.cpp b/Engine/Source/Developer/CollectionManager/Private/Collection.cpp
--- a/Engine/Source/Developer/CollectionManager/Private/Collection.cpp
+++ b/Engine/Source/Developer/CollectionManager/Private/Collection.cpp
@@ -136,26 +136,40 @@ bool FCollection::Load(FText& OutError)
 	{
 		const int32 NamesNum = FileContents.Num() - LineIndex;
 
-		TArray<FSoftObjectPath> Paths;
+		struct FItem
+		{
+			FSoftObjectPath Path;
+			int32 Hash;
+		};
+
+		TArray<FItem> Paths;
 		Paths.SetNum(NamesNum);
 
+		typedef decltype(ObjectSet)::KeyFuncsType KeyFuncsType;
 		// Name hashing to register new FName takes time
 		// Process as much as possible in multiple threads
 		ParallelFor(
 			NamesNum,
-			[this, &FileContents, &LineIndex, &Paths](int32 LocalLineIndex)
+			[&FileContents, &LineIndex, &Paths](int32 LocalLineIndex)
 			{
 				FStringView Line(FileContents[LineIndex + LocalLineIndex]);
-				Paths[LocalLineIndex] = FSoftObjectPath(Line.TrimStartAndEnd());
+				Paths[LocalLineIndex].Path = FSoftObjectPath(Line.TrimStartAndEnd());
+				Paths[LocalLineIndex].Hash = KeyFuncsType::GetKeyHash(KeyFuncsType::GetSetKey(Paths[LocalLineIndex].Path));
 			},
 			// Do not pay for scheduling cost if number of items is too low
 			NamesNum < 1000 ? EParallelForFlags::ForceSingleThread : EParallelForFlags::None
 		);
 
+		ObjectSet.Reserve(ObjectSet.Num() + Paths.Num());
 		// Static collection, a flat list of asset paths
-		for (FSoftObjectPath& Path : Paths)
+		for (FItem& Item : Paths)
 		{
-			AddObjectToCollection(Path);
+			if (Item.Path.IsNull())
+			{
+				continue;
+			}
+
+			ObjectSet.EmplaceByHash(Item.Hash, MoveTemp(Item.Path));
 		}
 	}
 	else
