Fixed reentry into FCollectionContainer lock when a long running source control operation causes the slow task or source control progress dialog to appear resulting in arbitrary UI code calling back into the collection container.

diff --git a/Engine/Source/Developer/CollectionManager/Private/CollectionContainer.cpp b/Engine/Source/Developer/CollectionManager/Private/CollectionContainer.cpp
--- a/Engine/Source/Developer/CollectionManager/Private/CollectionContainer.cpp
+++ b/Engine/Source/Developer/CollectionManager/Private/CollectionContainer.cpp
@@ -20,11 +20,129 @@
 
 #define LOCTEXT_NAMESPACE "CollectionManager"
 
+FCollectionRecursiveRwLock::FThreadLockDepths::FThreadLockDepths(void* TlsSlotValue)
+{
+	ThreadReadDepth = static_cast<HalfUPtrInt>(reinterpret_cast<UPTRINT>(TlsSlotValue));
+	ThreadWriteDepth = static_cast<HalfUPtrInt>(reinterpret_cast<UPTRINT>(TlsSlotValue) >> (sizeof(HalfUPtrInt) * 8));
+}
+
+void* FCollectionRecursiveRwLock::FThreadLockDepths::GetTlsSlotValue()
+{
+	UPTRINT ValueOut = static_cast<UPTRINT>(ThreadReadDepth) | (static_cast<UPTRINT>(ThreadWriteDepth) << (sizeof(HalfUPtrInt) * 8));
+	return reinterpret_cast<void*>(ValueOut);
+}
+
+FCollectionRecursiveRwLock::FCollectionRecursiveRwLock()
+{
+	TlsSlot = FPlatformTLS::AllocTlsSlot();
+
+	check(FPlatformTLS::IsValidTlsSlot(TlsSlot));
+}
+
+FCollectionRecursiveRwLock::~FCollectionRecursiveRwLock()
+{
+	check(FPlatformTLS::IsValidTlsSlot(TlsSlot));
+
+	FPlatformTLS::FreeTlsSlot(TlsSlot);
+}
+
+void FCollectionRecursiveRwLock::ReadLock()
+{
+	FThreadLockDepths LockDepths(FPlatformTLS::GetTlsValue(TlsSlot));
+	LockDepths.ThreadReadDepth++;
+	FPlatformTLS::SetTlsValue(TlsSlot, LockDepths.GetTlsSlotValue());
+
+	if (LockDepths.ThreadReadDepth + LockDepths.ThreadWriteDepth == 1)
+	{
+		RwLock.ReadLock();
+	}
+}
+
+void FCollectionRecursiveRwLock::WriteLock()
+{
+	// Some collection operations under write lock may cause Slate to tick if they show a dialog (e.g. FSlowTask),
+	// which in turn may enter back into the collection container to retrieve data to update the UI.
+	// Since this is difficult to avoid, we use a recursive read/write lock similar to the FPhysicsRwLock implementation.
+	FThreadLockDepths LockDepths(FPlatformTLS::GetTlsValue(TlsSlot));
+	LockDepths.ThreadWriteDepth++;
+	FPlatformTLS::SetTlsValue(TlsSlot, LockDepths.GetTlsSlotValue());
+
+	// If we have a read lock, we can't reuse it as a write lock without releasing it first.
+	// Call PromoteInterruptible to do so.
+	// If this situation arises because a thread with a read lock reenters and calls a high level function that requires
+	// a write lock, then we need to reconsider this solution (we'll likely need to lift all the slow tasks and source
+	// control operations out of the locks to prevent reentry altogether).
+	checkf(LockDepths.ThreadReadDepth == 0, TEXT("Attempting to acquire a write lock on a collection container when the same thread already holds a read lock higher up the call stack."));
+
+	if (LockDepths.ThreadReadDepth + LockDepths.ThreadWriteDepth == 1)
+	{
+		RwLock.WriteLock();
+	}
+}
+
+void FCollectionRecursiveRwLock::ReadUnlock()
+{
+	FThreadLockDepths LockDepths(FPlatformTLS::GetTlsValue(TlsSlot));
+	LockDepths.ThreadReadDepth--;
+	FPlatformTLS::SetTlsValue(TlsSlot, LockDepths.GetTlsSlotValue());
+
+	if (LockDepths.ThreadWriteDepth + LockDepths.ThreadReadDepth == 0)
+	{
+		RwLock.ReadUnlock();
+	}
+}
+
+void FCollectionRecursiveRwLock::WriteUnlock()
+{
+	FThreadLockDepths LockDepths(FPlatformTLS::GetTlsValue(TlsSlot));
+	LockDepths.ThreadWriteDepth--;
+	FPlatformTLS::SetTlsValue(TlsSlot, LockDepths.GetTlsSlotValue());
+
+	if (LockDepths.ThreadWriteDepth + LockDepths.ThreadReadDepth == 0)
+	{
+		RwLock.WriteUnlock();
+	}
+}
+
+bool FCollectionRecursiveRwLock::PromoteInterruptible()
+{
+	FThreadLockDepths LockDepths(FPlatformTLS::GetTlsValue(TlsSlot));
+
+	if (LockDepths.ThreadWriteDepth > 0 && LockDepths.ThreadReadDepth == 0)
+	{
+		// Already promoted.
+		return true;
+	}
+
+	// Expecting this to be called from a thread with a read lock.
+	checkf(LockDepths.ThreadReadDepth > 0, TEXT("Attempting to promote a read lock to a write lock on a collection container when the thread isn't holding a read lock."));
+
+	// We cannot promote if we're not the topmost holder of the lock on this thread since we need to release the read lock and there might be another
+	// thread trying to enter write lock.
+	// 
+	// Any code directly calling PromoteInterruptible is expecting to be pre-empted and will reconfirm its invariants and reacquire any pointers which may have been invalidated.
+	// 
+	// If there is another scope higher up on this thread holding a read lock, it is not expecting to be pre-empted and will not be prepared for its invariants to be broken or
+	// its pointers to be invalidated.
+	if (LockDepths.ThreadWriteDepth == 0 && LockDepths.ThreadReadDepth == 1)
+	{
+		LockDepths.ThreadWriteDepth = 1;
+		LockDepths.ThreadReadDepth = 0;
+		FPlatformTLS::SetTlsValue(TlsSlot, LockDepths.GetTlsSlotValue());
+
+		RwLock.ReadUnlock();
+		RwLock.WriteLock();
+		return true;
+	}
+
+	return false;
+}
+
 // Base class for lock hierarchy. When used as a function parameter it means the called must hold at least a read lock
-class FCollectionLock
+class FCollectionScopeLock
 {
 protected:
-	UE_NODISCARD_CTOR explicit FCollectionLock(FRWLock& InLockObject, bool InWriteLock)
+	UE_NODISCARD_CTOR explicit FCollectionScopeLock(FCollectionRecursiveRwLock& InLockObject, bool InWriteLock)
 	: LockObject(InLockObject)
 	, bWriteLock(InWriteLock)
 	{
@@ -38,18 +156,26 @@ protected:
 		}
 	}
 	
-	// Promoted the lock from read to write, possibly being interrupted by another writer in between
-	void PromoteInterruptible()
+	// Promote the lock from read to write, possibly being interrupted by another writer in between.
+	// Returns false if the lock cannot be promoted, which can happen if the thread already holds a read
+	// lock then reenters and tries to promote to a write lock.
+	[[nodiscard]] bool PromoteInterruptible()
 	{
-		if(!bWriteLock)
+		if (bWriteLock)
+		{
+			// Already promoted.
+			return true;
+		}
+
+		if (LockObject.PromoteInterruptible())
 		{
-			LockObject.ReadUnlock();
-			LockObject.WriteLock();
 			bWriteLock = true;
 		}
+
+		return bWriteLock;
 	}
 	
-	~FCollectionLock()
+	~FCollectionScopeLock()
 	{
 		if(bWriteLock)
 		{
@@ -68,42 +194,42 @@ protected:
 	}
 	
 private:
-	UE_NONCOPYABLE(FCollectionLock);
-	FRWLock& LockObject;
+	UE_NONCOPYABLE(FCollectionScopeLock);
+	FCollectionRecursiveRwLock& LockObject;
 	bool bWriteLock = false;
 };
 
 // Scoped lock type used to hold lock and to tag methods which should at least hold a read lock 
-class FCollectionLock_Read : public FCollectionLock
+class FCollectionScopeLock_Read : public FCollectionScopeLock
 {
 public:
-	UE_NODISCARD_CTOR explicit FCollectionLock_Read(FRWLock& InLockObject)
-	: FCollectionLock(InLockObject, false)
+	UE_NODISCARD_CTOR explicit FCollectionScopeLock_Read(FCollectionRecursiveRwLock& InLockObject)
+	: FCollectionScopeLock(InLockObject, false)
 	{
 	}
 };
 
 // A lock on the collection container which begins in a read only state and can be promoted into a write lock with potential interruption in between 
-class FCollectionLock_RW : public FCollectionLock
+class FCollectionScopeLock_RW : public FCollectionScopeLock
 {
 public:
-	UE_NODISCARD_CTOR explicit FCollectionLock_RW(FRWLock& InLockObject, bool bWrite = false)
-	: FCollectionLock(InLockObject, bWrite)
+	UE_NODISCARD_CTOR explicit FCollectionScopeLock_RW(FCollectionRecursiveRwLock& InLockObject, bool bWrite = false)
+	: FCollectionScopeLock(InLockObject, bWrite)
 	{
 	}
 	
 	// Promoted the lock from read to write, possibly being interrupted by another writer in between
-	using FCollectionLock::PromoteInterruptible;
+	using FCollectionScopeLock::PromoteInterruptible;
 	// Used for assertions to confirm that the correct kind of lock has been taken
-	using FCollectionLock::IsWriteLock;
+	using FCollectionScopeLock::IsWriteLock;
 };
 
 // Write lock on the collection container
-class FCollectionLock_Write : public FCollectionLock_RW
+class FCollectionScopeLock_Write : public FCollectionScopeLock_RW
 {
 public:
-	UE_NODISCARD_CTOR explicit FCollectionLock_Write(FRWLock& InLockObject)
-	: FCollectionLock_RW(InLockObject, true)
+	UE_NODISCARD_CTOR explicit FCollectionScopeLock_Write(FCollectionRecursiveRwLock& InLockObject)
+	: FCollectionScopeLock_RW(InLockObject, true)
 	{
 	}
 };
@@ -118,19 +244,19 @@ public:
 	 * Dirty the parts of the cache that need to change when a collection is added to our collection container.
 	 * The collection container must be locked.
 	 */
-	void HandleCollectionAdded(FCollectionLock_Write&);
+	void HandleCollectionAdded(FCollectionScopeLock_Write&);
 	
 	/** 
 	 * Dirty the parts of the cache that need to change when a collection is removed from our collection container 
 	 * The collection container must be locked.
 	 */
-	void HandleCollectionRemoved(FCollectionLock_Write&);
+	void HandleCollectionRemoved(FCollectionScopeLock_Write&);
 
 	/** 
 	 * Dirty the parts of the cache that need to change when a collection is modified 
 	 * The collcetion container must be lockedl
 	 */
-	void HandleCollectionChanged(FCollectionLock_Write&);
+	void HandleCollectionChanged(FCollectionScopeLock_Write&);
 
 	/** 
 	 * Update the given dirty parts of the cache based on which parts will be accessed while the given lock is held.
@@ -139,34 +265,37 @@ public:
 	 * The calling thread may be interrupted by another write operation during the promotion operation.
 	 * Therefore, caches should be updated as early as possible in order to prevent invalidation of state.
 	 * 
+	 * Returns false if the lock cannot be promoted preventing the cache from updating, which can happen if the
+	 * thread already holds a read lock then reenters and tries to promote to a write lock.
+	 * 
 	 * This function is used rather than updating the caches in the Get* functions to prevent issues with pre-emption 
 	 * on the lock upgrade deep into a method.
 	 */
-	void UpdateCaches(FCollectionLock_RW& InGuard, ECollectionCacheFlags Flags);
+	bool UpdateCaches(FCollectionScopeLock_RW& InGuard, ECollectionCacheFlags Flags);
 
 	/** 
 	 * Access the CachedCollectionNamesFromGuids map, asserting that it is up-to-date.
 	 * The collection container must be read-locked.
 	 */
-	const TMap<FGuid, FCollectionNameType>& GetCachedCollectionNamesFromGuids(FCollectionLock&) const;
+	const TMap<FGuid, FCollectionNameType>& GetCachedCollectionNamesFromGuids(FCollectionScopeLock&) const;
 
 	/** 
 	 * Access the CachedObjects map, asserting that it is up-to-date.
 	 * The collection container must be read-locked.
 	 */
-	const TMap<FSoftObjectPath, TArray<FObjectCollectionInfo>>& GetCachedObjects(FCollectionLock&) const;
+	const TMap<FSoftObjectPath, TArray<FObjectCollectionInfo>>& GetCachedObjects(FCollectionScopeLock&) const;
 
 	/** 
 	 * Access the CachedHierarchy map, asserting that it is up-to-date.
 	 * The collection container must be read-locked.
 	 */
-	const TMap<FGuid, TArray<FGuid>>& GetCachedHierarchy(FCollectionLock&) const;
+	const TMap<FGuid, TArray<FGuid>>& GetCachedHierarchy(FCollectionScopeLock&) const;
 
 	/** 
 	 * Access the CachedColors array, asserting that it is up-to-date 
 	 * The collection container must be read-locked.
 	 */
-	const TArray<FLinearColor>& GetCachedColors(FCollectionLock&) const;
+	const TArray<FLinearColor>& GetCachedColors(FCollectionScopeLock&) const;
 
 	enum class ERecursiveWorkerFlowControl : uint8
 	{
@@ -180,7 +309,7 @@ public:
 	 * Perform a recursive operation on the given collection and optionally its parents and children.
 	 * The collection container must be read-locked and UpdateCaches must be called for names and hierarchy.
 	 */
-	void RecursionHelper_DoWork(FCollectionLock&, const FCollectionNameType& InCollectionKey, const ECollectionRecursionFlags::Flags InRecursionMode, FRecursiveWorkerFunc InWorkerFunc) const;
+	void RecursionHelper_DoWork(FCollectionScopeLock&, const FCollectionNameType& InCollectionKey, const ECollectionRecursionFlags::Flags InRecursionMode, FRecursiveWorkerFunc InWorkerFunc) const;
 
 private:
 	/** Reference to the collections that are currently available in our owner collection container */
@@ -201,8 +330,8 @@ private:
 	/** Which parts of the cache are dirty */
 	ECollectionCacheFlags DirtyFlags = ECollectionCacheFlags::All;
 
-	ERecursiveWorkerFlowControl RecursionHelper_DoWorkOnParents(FCollectionLock&, const FCollectionNameType& InCollectionKey, FRecursiveWorkerFunc InWorkerFunc) const;
-	ERecursiveWorkerFlowControl RecursionHelper_DoWorkOnChildren(FCollectionLock&, const FCollectionNameType& InCollectionKey, FRecursiveWorkerFunc InWorkerFunc) const;
+	ERecursiveWorkerFlowControl RecursionHelper_DoWorkOnParents(FCollectionScopeLock&, const FCollectionNameType& InCollectionKey, FRecursiveWorkerFunc InWorkerFunc) const;
+	ERecursiveWorkerFlowControl RecursionHelper_DoWorkOnChildren(FCollectionScopeLock&, const FCollectionNameType& InCollectionKey, FRecursiveWorkerFunc InWorkerFunc) const;
 };
 
 FCollectionContainerCache::FCollectionContainerCache(TMap<FCollectionNameType, TSharedRef<FCollection>>* InAvailableCollections)
@@ -210,22 +339,22 @@ FCollectionContainerCache::FCollectionContainerCache(TMap<FCollectionNameType, T
 {
 }
 
-void FCollectionContainerCache::HandleCollectionAdded(FCollectionLock_Write&)
+void FCollectionContainerCache::HandleCollectionAdded(FCollectionScopeLock_Write&)
 {
 	DirtyFlags |= ECollectionCacheFlags::Names;
 }
 
-void FCollectionContainerCache::HandleCollectionRemoved(FCollectionLock_Write&)
+void FCollectionContainerCache::HandleCollectionRemoved(FCollectionScopeLock_Write&)
 {
 	DirtyFlags |= ECollectionCacheFlags::All;
 }
 
-void FCollectionContainerCache::HandleCollectionChanged(FCollectionLock_Write&)
+void FCollectionContainerCache::HandleCollectionChanged(FCollectionScopeLock_Write&)
 {
 	DirtyFlags |= ECollectionCacheFlags::Objects | ECollectionCacheFlags::Hierarchy | ECollectionCacheFlags::Colors;
 }
 
-void FCollectionContainerCache::UpdateCaches(FCollectionLock_RW& InGuard, ECollectionCacheFlags ToUpdate)
+bool FCollectionContainerCache::UpdateCaches(FCollectionScopeLock_RW& InGuard, ECollectionCacheFlags ToUpdate)
 {
 	// Updating objects or hierarchy requires name mapping
 	if (EnumHasAnyFlags(ToUpdate, ECollectionCacheFlags::Hierarchy | ECollectionCacheFlags::Objects))
@@ -241,13 +370,18 @@ void FCollectionContainerCache::UpdateCaches(FCollectionLock_RW& InGuard, EColle
 
 	if (EnumHasAnyFlags(DirtyFlags, ToUpdate))
 	{
-		InGuard.PromoteInterruptible();
+		if (!InGuard.PromoteInterruptible())
+		{
+			// We assume get operations that require updating the cache on a thread we've reentered are originating from Slate bindings,
+			// as such using a stale cache should be ok since such operations should succeed again when Slate ticks after the stack unwinds.
+			return false;
+		}
 	}
 
 	if (!EnumHasAnyFlags(DirtyFlags, ToUpdate))
 	{
 		// Caches we care about were updated while we switched locks
-		return;
+		return true;
 	}
 
 	// Limit updates to what's dirty 
@@ -365,9 +499,10 @@ void FCollectionContainerCache::UpdateCaches(FCollectionLock_RW& InGuard, EColle
 	}
 
 	UE_LOG(LogCollectionManager, Verbose, TEXT("Rebuilt caches for %d collections in in %0.6f seconds"), AvailableCollections->Num(), FPlatformTime::Seconds() - CacheStartTime);
+	return true;
 }
 
-void FCollectionContainerCache::RecursionHelper_DoWork(FCollectionLock& Guard,
+void FCollectionContainerCache::RecursionHelper_DoWork(FCollectionScopeLock& Guard,
 	const FCollectionNameType& InCollectionKey,
 	const ECollectionRecursionFlags::Flags InRecursionMode,
 	FRecursiveWorkerFunc InWorkerFunc) const
@@ -391,7 +526,7 @@ void FCollectionContainerCache::RecursionHelper_DoWork(FCollectionLock& Guard,
 }
 
 FCollectionContainerCache::ERecursiveWorkerFlowControl FCollectionContainerCache::RecursionHelper_DoWorkOnParents(
-	FCollectionLock& Guard, const FCollectionNameType& InCollectionKey, FRecursiveWorkerFunc InWorkerFunc) const
+	FCollectionScopeLock& Guard, const FCollectionNameType& InCollectionKey, FRecursiveWorkerFunc InWorkerFunc) const
 {
 	const TSharedRef<FCollection>* const CollectionRefPtr = AvailableCollections->Find(InCollectionKey);
 	if (CollectionRefPtr)
@@ -412,7 +547,7 @@ FCollectionContainerCache::ERecursiveWorkerFlowControl FCollectionContainerCache
 }
 
 FCollectionContainerCache::ERecursiveWorkerFlowControl FCollectionContainerCache::RecursionHelper_DoWorkOnChildren(
-	FCollectionLock& Guard, const FCollectionNameType& InCollectionKey, FRecursiveWorkerFunc InWorkerFunc) const
+	FCollectionScopeLock& Guard, const FCollectionNameType& InCollectionKey, FRecursiveWorkerFunc InWorkerFunc) const
 {
 	const TSharedRef<FCollection>* const CollectionRefPtr = AvailableCollections->Find(InCollectionKey);
 	if (CollectionRefPtr)
@@ -441,25 +576,25 @@ FCollectionContainerCache::ERecursiveWorkerFlowControl FCollectionContainerCache
 	return ERecursiveWorkerFlowControl::Continue;
 }
 
-const TMap<FGuid, FCollectionNameType>& FCollectionContainerCache::GetCachedCollectionNamesFromGuids(FCollectionLock&) const 
+const TMap<FGuid, FCollectionNameType>& FCollectionContainerCache::GetCachedCollectionNamesFromGuids(FCollectionScopeLock&) const
 {
 	checkf(!EnumHasAnyFlags(DirtyFlags, ECollectionCacheFlags::Names), TEXT("Accessed guid->name map without updating cache"));
 	return CachedCollectionNamesFromGuids_Internal;
 }
 
-const TMap<FSoftObjectPath, TArray<FObjectCollectionInfo>>& FCollectionContainerCache::GetCachedObjects(FCollectionLock&) const 
+const TMap<FSoftObjectPath, TArray<FObjectCollectionInfo>>& FCollectionContainerCache::GetCachedObjects(FCollectionScopeLock&) const
 {
 	checkf(!EnumHasAnyFlags(DirtyFlags, ECollectionCacheFlags::Objects), TEXT("Accessd object->collection map without updating cache"));
 	return CachedObjects_Internal;
 }
 
-const TMap<FGuid, TArray<FGuid>>& FCollectionContainerCache::GetCachedHierarchy(FCollectionLock&) const 
+const TMap<FGuid, TArray<FGuid>>& FCollectionContainerCache::GetCachedHierarchy(FCollectionScopeLock&) const
 {
 	checkf(!EnumHasAnyFlags(DirtyFlags, ECollectionCacheFlags::Hierarchy), TEXT("Accessed collection hierarchy map without updating cache"));
 	return CachedHierarchy_Internal;
 }
 
-const TArray<FLinearColor>& FCollectionContainerCache::GetCachedColors(FCollectionLock&) const 
+const TArray<FLinearColor>& FCollectionContainerCache::GetCachedColors(FCollectionScopeLock&) const
 {
 	checkf(!EnumHasAnyFlags(DirtyFlags, ECollectionCacheFlags::Colors), TEXT("Accessed collection colors without updating cache"));
 	return CachedColors_Internal;
@@ -504,7 +639,7 @@ FCollectionContainer::~FCollectionContainer() = default;
 
 bool FCollectionContainer::IsReadOnly(ECollectionShareType::Type ShareType) const
 {
-	FCollectionLock_Read Guard(Lock);
+	FCollectionScopeLock_Read Guard(Lock);
 
 	return IsReadOnly(Guard, ShareType);
 }
@@ -513,7 +648,7 @@ void FCollectionContainer::SetReadOnly(ECollectionShareType::Type ShareType, boo
 {
 	check(ShareType <= ECollectionShareType::CST_All);
 
-	FCollectionLock_Write Guard(Lock);
+	FCollectionScopeLock_Write Guard(Lock);
 
 	const uint8 ReadOnlyMask = GetReadOnlyMask(ShareType);
 	if (bReadOnly)
@@ -528,14 +663,14 @@ void FCollectionContainer::SetReadOnly(ECollectionShareType::Type ShareType, boo
 
 bool FCollectionContainer::IsHidden() const
 {
-	FCollectionLock_Read Guard(Lock);
+	FCollectionScopeLock_Read Guard(Lock);
 	return bIsHidden;
 }
 
 void FCollectionContainer::SetHidden(bool bHidden)
 {
 	{
-		FCollectionLock_Write Guard(Lock);
+		FCollectionScopeLock_Write Guard(Lock);
 
 		if (bIsHidden == bHidden)
 		{
@@ -550,13 +685,13 @@ void FCollectionContainer::SetHidden(bool bHidden)
 
 bool FCollectionContainer::HasCollections() const
 {
-	FCollectionLock_Read Guard(Lock);
+	FCollectionScopeLock_Read Guard(Lock);
 	return AvailableCollections.Num() > 0;
 }
 
 void FCollectionContainer::GetCollections(TArray<FCollectionNameType>& OutCollections) const
 {
-	FCollectionLock_Read Guard(Lock);
+	FCollectionScopeLock_Read Guard(Lock);
 	OutCollections.Reserve(AvailableCollections.Num());
 	for (const TPair<FCollectionNameType, TSharedRef<FCollection>>& AvailableCollection : AvailableCollections)
 	{
@@ -567,7 +702,7 @@ void FCollectionContainer::GetCollections(TArray<FCollectionNameType>& OutCollec
 
 void FCollectionContainer::GetCollections(FName CollectionName, TArray<FCollectionNameType>& OutCollections) const
 {
-	FCollectionLock_Read Guard(Lock);
+	FCollectionScopeLock_Read Guard(Lock);
 	for (int32 CacheIdx = 0; CacheIdx < ECollectionShareType::CST_All; ++CacheIdx)
 	{
 		if (AvailableCollections.Contains(FCollectionNameType(CollectionName, ECollectionShareType::Type(CacheIdx))))
@@ -579,7 +714,7 @@ void FCollectionContainer::GetCollections(FName CollectionName, TArray<FCollecti
 
 void FCollectionContainer::GetCollectionNames(ECollectionShareType::Type ShareType, TArray<FName>& CollectionNames) const
 {
-	FCollectionLock_Read Guard(Lock);
+	FCollectionScopeLock_Read Guard(Lock);
 	for (const TPair<FCollectionNameType, TSharedRef<FCollection>>& AvailableCollection : AvailableCollections)
 	{
 		const FCollectionNameType& CollectionKey = AvailableCollection.Key;
@@ -592,7 +727,7 @@ void FCollectionContainer::GetCollectionNames(ECollectionShareType::Type ShareTy
 
 void FCollectionContainer::GetRootCollections(TArray<FCollectionNameType>& OutCollections) const
 {
-	FCollectionLock_RW Guard(Lock);
+	FCollectionScopeLock_RW Guard(Lock);
 	CollectionCache->UpdateCaches(Guard, ECollectionCacheFlags::Names);
 	const TMap<FGuid, FCollectionNameType>& CachedCollectionNamesFromGuids = CollectionCache->GetCachedCollectionNamesFromGuids(Guard);
 
@@ -612,7 +747,7 @@ void FCollectionContainer::GetRootCollections(TArray<FCollectionNameType>& OutCo
 
 void FCollectionContainer::GetRootCollectionNames(ECollectionShareType::Type ShareType, TArray<FName>& CollectionNames) const
 {
-	FCollectionLock_RW Guard(Lock);
+	FCollectionScopeLock_RW Guard(Lock);
 	CollectionCache->UpdateCaches(Guard, ECollectionCacheFlags::Names);
 	const TMap<FGuid, FCollectionNameType>& CachedCollectionNamesFromGuids = CollectionCache->GetCachedCollectionNamesFromGuids(Guard);
 
@@ -634,7 +769,7 @@ void FCollectionContainer::GetRootCollectionNames(ECollectionShareType::Type Sha
 
 void FCollectionContainer::GetChildCollections(FName CollectionName, ECollectionShareType::Type ShareType, TArray<FCollectionNameType>& OutCollections) const
 {
-	FCollectionLock_RW Guard(Lock);
+	FCollectionScopeLock_RW Guard(Lock);
 	CollectionCache->UpdateCaches(Guard, ECollectionCacheFlags::Names | ECollectionCacheFlags::Hierarchy);
 
 	const TMap<FGuid, FCollectionNameType>& CachedCollectionNamesFromGuids = CollectionCache->GetCachedCollectionNamesFromGuids(Guard);
@@ -671,7 +806,7 @@ void FCollectionContainer::GetChildCollections(FName CollectionName, ECollection
 
 void FCollectionContainer::GetChildCollectionNames(FName CollectionName, ECollectionShareType::Type ShareType, ECollectionShareType::Type ChildShareType, TArray<FName>& CollectionNames) const
 {
-	FCollectionLock_RW Guard(Lock);
+	FCollectionScopeLock_RW Guard(Lock);
 	CollectionCache->UpdateCaches(Guard, ECollectionCacheFlags::Names | ECollectionCacheFlags::Hierarchy);
 	const TMap<FGuid, FCollectionNameType>& CachedCollectionNamesFromGuids = CollectionCache->GetCachedCollectionNamesFromGuids(Guard);
 	const TMap<FGuid, TArray<FGuid>>& CachedHierarchy = CollectionCache->GetCachedHierarchy(Guard);
@@ -707,7 +842,7 @@ void FCollectionContainer::GetChildCollectionNames(FName CollectionName, ECollec
 
 TOptional<FCollectionNameType> FCollectionContainer::GetParentCollection(FName CollectionName, ECollectionShareType::Type ShareType) const
 {
-	FCollectionLock_RW Guard(Lock);
+	FCollectionScopeLock_RW Guard(Lock);
 	const TSharedRef<FCollection>* const CollectionRefPtr = AvailableCollections.Find(FCollectionNameType(CollectionName, ShareType));
 	if (!CollectionRefPtr)
 	{
@@ -726,11 +861,11 @@ TOptional<FCollectionNameType> FCollectionContainer::GetParentCollection(FName C
 
 bool FCollectionContainer::CollectionExists(FName CollectionName, ECollectionShareType::Type ShareType) const
 {
-	FCollectionLock_Read Guard(Lock);
+	FCollectionScopeLock_Read Guard(Lock);
 	return CollectionExists_Locked(Guard, CollectionName, ShareType);
 }
 
-bool FCollectionContainer::CollectionExists_Locked(FCollectionLock&, FName CollectionName, ECollectionShareType::Type ShareType) const
+bool FCollectionContainer::CollectionExists_Locked(FCollectionScopeLock&, FName CollectionName, ECollectionShareType::Type ShareType) const
 {
 	if (ShareType == ECollectionShareType::CST_All)
 	{
@@ -755,7 +890,7 @@ bool FCollectionContainer::CollectionExists_Locked(FCollectionLock&, FName Colle
 
 bool FCollectionContainer::GetAssetsInCollection(FName CollectionName, ECollectionShareType::Type ShareType, TArray<FSoftObjectPath>& AssetsPaths, ECollectionRecursionFlags::Flags RecursionMode) const
 {
-	FCollectionLock_RW Guard(Lock);
+	FCollectionScopeLock_RW Guard(Lock);
 	CollectionCache->UpdateCaches(Guard, ECollectionCacheFlags::RecursionWorker);
 	bool bFoundAssets = false;
 
@@ -788,7 +923,7 @@ bool FCollectionContainer::GetAssetsInCollection(FName CollectionName, ECollecti
 
 bool FCollectionContainer::GetClassesInCollection(FName CollectionName, ECollectionShareType::Type ShareType, TArray<FTopLevelAssetPath>& ClassPaths, ECollectionRecursionFlags::Flags RecursionMode) const
 {
-	FCollectionLock_RW Guard(Lock);
+	FCollectionScopeLock_RW Guard(Lock);
 	CollectionCache->UpdateCaches(Guard, ECollectionCacheFlags::RecursionWorker);
 	bool bFoundClasses = false;
 
@@ -821,7 +956,7 @@ bool FCollectionContainer::GetClassesInCollection(FName CollectionName, ECollect
 
 bool FCollectionContainer::GetObjectsInCollection(FName CollectionName, ECollectionShareType::Type ShareType, TArray<FSoftObjectPath>& ObjectPaths, ECollectionRecursionFlags::Flags RecursionMode) const
 {
-	FCollectionLock_RW Guard(Lock);
+	FCollectionScopeLock_RW Guard(Lock);
 	CollectionCache->UpdateCaches(Guard, ECollectionCacheFlags::RecursionWorker);
 	bool bFoundObjects = false;
 
@@ -854,7 +989,7 @@ bool FCollectionContainer::GetObjectsInCollection(FName CollectionName, ECollect
 
 void FCollectionContainer::GetCollectionsContainingObject(const FSoftObjectPath& ObjectPath, ECollectionShareType::Type ShareType, TArray<FName>& OutCollectionNames, ECollectionRecursionFlags::Flags RecursionMode) const
 {
-	FCollectionLock_RW Guard(Lock);
+	FCollectionScopeLock_RW Guard(Lock);
 	CollectionCache->UpdateCaches(Guard, ECollectionCacheFlags::Objects);
 	const TMap<FSoftObjectPath, TArray<FObjectCollectionInfo>>& CachedObjects = CollectionCache->GetCachedObjects(Guard);
 
@@ -873,7 +1008,7 @@ void FCollectionContainer::GetCollectionsContainingObject(const FSoftObjectPath&
 
 void FCollectionContainer::GetCollectionsContainingObject(const FSoftObjectPath& ObjectPath, TArray<FCollectionNameType>& OutCollections, ECollectionRecursionFlags::Flags RecursionMode) const
 {
-	FCollectionLock_RW Guard(Lock);
+	FCollectionScopeLock_RW Guard(Lock);
 	CollectionCache->UpdateCaches(Guard, ECollectionCacheFlags::Objects);
 	const TMap<FSoftObjectPath, TArray<FObjectCollectionInfo>>& CachedObjects = CollectionCache->GetCachedObjects(Guard);
 
@@ -893,7 +1028,7 @@ void FCollectionContainer::GetCollectionsContainingObject(const FSoftObjectPath&
 
 void FCollectionContainer::GetCollectionsContainingObjects(const TArray<FSoftObjectPath>& ObjectPaths, TMap<FCollectionNameType, TArray<FSoftObjectPath>>& OutCollectionsAndMatchedObjects, ECollectionRecursionFlags::Flags RecursionMode) const
 {
-	FCollectionLock_RW Guard(Lock);
+	FCollectionScopeLock_RW Guard(Lock);
 	CollectionCache->UpdateCaches(Guard, ECollectionCacheFlags::Objects);
 	const TMap<FSoftObjectPath, TArray<FObjectCollectionInfo>>& CachedObjects = CollectionCache->GetCachedObjects(Guard);
 
@@ -916,7 +1051,7 @@ void FCollectionContainer::GetCollectionsContainingObjects(const TArray<FSoftObj
 
 FString FCollectionContainer::GetCollectionsStringForObject(const FSoftObjectPath& ObjectPath, ECollectionShareType::Type ShareType, ECollectionRecursionFlags::Flags RecursionMode, bool bFullPaths) const
 {
-	FCollectionLock_RW Guard(Lock);
+	FCollectionScopeLock_RW Guard(Lock);
 	CollectionCache->UpdateCaches(Guard, ECollectionCacheFlags::Objects | ECollectionCacheFlags::RecursionWorker);
 	const TMap<FSoftObjectPath, TArray<FObjectCollectionInfo>>& CachedObjects = CollectionCache->GetCachedObjects(Guard);
 
@@ -973,7 +1108,7 @@ FString FCollectionContainer::MakeCollectionPath(FName CollectionName, ECollecti
 
 void FCollectionContainer::CreateUniqueCollectionName(FName BaseName, ECollectionShareType::Type ShareType, FName& OutCollectionName) const
 {
-	FCollectionLock_Read Guard(Lock);
+	FCollectionScopeLock_Read Guard(Lock);
 	int32 IntSuffix = 1;
 	bool CollectionAlreadyExists = false;
 	do
@@ -1056,7 +1191,7 @@ bool FCollectionContainer::CreateCollection(FName CollectionName, ECollectionSha
 	}
 
 	{
-		FCollectionLock_Write Guard(Lock);
+		FCollectionScopeLock_Write Guard(Lock);
 
 		if (!ValidateWritable(Guard, ShareType, OutError))
 		{
@@ -1108,7 +1243,7 @@ bool FCollectionContainer::RenameCollection(FName CurrentCollectionName, ECollec
 	const FCollectionNameType OriginalCollectionKey(CurrentCollectionName, CurrentShareType);
 	const FCollectionNameType NewCollectionKey(NewCollectionName, NewShareType);
 	{
-		FCollectionLock_Write Guard(Lock);
+		FCollectionScopeLock_Write Guard(Lock);
 
 		if (!ValidateWritable(Guard, CurrentShareType, OutError) || !ValidateWritable(Guard, NewShareType, OutError))
 		{
@@ -1187,14 +1322,14 @@ bool FCollectionContainer::ReparentCollection(FName CollectionName, ECollectionS
 	TOptional<FCollectionNameType> OldParentCollectionKey;
 	TOptional<FCollectionNameType> NewParentCollectionKey;
 	{
-		FCollectionLock_Write Guard(Lock);
+		FCollectionScopeLock_Write Guard(Lock);
 
 		if (!ValidateWritable(Guard, ShareType, OutError) || (!ParentCollectionName.IsNone() && !ValidateWritable(Guard, ParentShareType, OutError)))
 		{
 			return false;
 		}
 
-		CollectionCache->UpdateCaches(Guard, ECollectionCacheFlags::RecursionWorker);
+		verifyf(CollectionCache->UpdateCaches(Guard, ECollectionCacheFlags::RecursionWorker), TEXT("UpdateCaches must be called within a write lock to guarantee subsequent usage to function as expected"));
 
 		TSharedRef<FCollection>* const CollectionRefPtr = AvailableCollections.Find(CollectionKey);
 		if (!CollectionRefPtr)
@@ -1270,7 +1405,7 @@ bool FCollectionContainer::ReparentCollection(FName CollectionName, ECollectionS
 		}
 
 		CollectionCache->HandleCollectionChanged(Guard);
-		CollectionCache->UpdateCaches(Guard, ECollectionCacheFlags::Names);
+		verifyf(CollectionCache->UpdateCaches(Guard, ECollectionCacheFlags::Names), TEXT("UpdateCaches must be called within a write lock to guarantee subsequent usage to function as expected"));
 
 		// Find the old parent so we can notify about the change
 		{
@@ -1293,7 +1428,7 @@ bool FCollectionContainer::DestroyCollection(FName CollectionName, ECollectionSh
 {
 	const FCollectionNameType CollectionKey(CollectionName, ShareType);
 	{
-		FCollectionLock_Write Guard(Lock);
+		FCollectionScopeLock_Write Guard(Lock);
 
 		if (!ValidateWritable(Guard, ShareType, OutError))
 		{
@@ -1344,7 +1479,7 @@ bool FCollectionContainer::AddToCollection(FName CollectionName, ECollectionShar
 
 	const FCollectionNameType CollectionKey(CollectionName, ShareType);
 	{
-		FCollectionLock_Write Guard(Lock);
+		FCollectionScopeLock_Write Guard(Lock);
 
 		if (!ValidateWritable(Guard, ShareType, OutError))
 		{
@@ -1436,7 +1571,7 @@ bool FCollectionContainer::RemoveFromCollection(FName CollectionName, ECollectio
 
 	const FCollectionNameType CollectionKey(CollectionName, ShareType);
 	{
-		FCollectionLock_Write Guard(Lock);
+		FCollectionScopeLock_Write Guard(Lock);
 
 		if (!ValidateWritable(Guard, ShareType, OutError))
 		{
@@ -1513,7 +1648,7 @@ bool FCollectionContainer::SetDynamicQueryText(FName CollectionName, ECollection
 {
 	const FCollectionNameType CollectionKey(CollectionName, ShareType);
 	{
-		FCollectionLock_Write Guard(Lock);
+		FCollectionScopeLock_Write Guard(Lock);
 
 		if (!ValidateWritable(Guard, ShareType, OutError))
 		{
@@ -1569,7 +1704,7 @@ bool FCollectionContainer::GetDynamicQueryText(FName CollectionName, ECollection
 		return false;
 	}
 
-	FCollectionLock_Read Guard(Lock);
+	FCollectionScopeLock_Read Guard(Lock);
 	const FCollectionNameType CollectionKey(CollectionName, ShareType);
 	const TSharedRef<FCollection>* const CollectionRefPtr = AvailableCollections.Find(CollectionKey);
 	if (!CollectionRefPtr)
@@ -1607,7 +1742,7 @@ bool FCollectionContainer::TestDynamicQuery(FName CollectionName, ECollectionSha
 		return false;
 	}
 
-	FCollectionLock_Read Guard(Lock);
+	FCollectionScopeLock_Read Guard(Lock);
 	const FCollectionNameType CollectionKey(CollectionName, ShareType);
 	const TSharedRef<FCollection>* const CollectionRefPtr = AvailableCollections.Find(CollectionKey);
 	if (!CollectionRefPtr)
@@ -1638,7 +1773,7 @@ bool FCollectionContainer::EmptyCollection(FName CollectionName, ECollectionShar
 {
 	const FCollectionNameType CollectionKey(CollectionName, ShareType);
 	{
-		FCollectionLock_Write Guard(Lock);
+		FCollectionScopeLock_Write Guard(Lock);
 
 		if (!ValidateWritable(Guard, ShareType, OutError))
 		{
@@ -1683,7 +1818,7 @@ bool FCollectionContainer::SaveCollection(FName CollectionName, ECollectionShare
 {
 	const FCollectionNameType CollectionKey(CollectionName, ShareType);
 	{
-		FCollectionLock_Write Guard(Lock);
+		FCollectionScopeLock_Write Guard(Lock);
 
 		if (!ValidateWritable(Guard, ShareType, OutError))
 		{
@@ -1729,7 +1864,7 @@ bool FCollectionContainer::UpdateCollection(FName CollectionName, ECollectionSha
 {
 	const FCollectionNameType CollectionKey(CollectionName, ShareType);
 	{
-		FCollectionLock_Write Guard(Lock);
+		FCollectionScopeLock_Write Guard(Lock);
 
 		if (!ValidateWritable(Guard, ShareType, OutError))
 		{
@@ -1775,7 +1910,7 @@ bool FCollectionContainer::GetCollectionStatusInfo(FName CollectionName, ECollec
 		return false;
 	}
 
-	FCollectionLock_Read Guard(Lock);
+	FCollectionScopeLock_Read Guard(Lock);
 	const FCollectionNameType CollectionKey(CollectionName, ShareType);
 	const TSharedRef<FCollection>* const CollectionRefPtr = AvailableCollections.Find(CollectionKey);
 	if (CollectionRefPtr)
@@ -1796,7 +1931,7 @@ bool FCollectionContainer::GetCollectionStatusInfo(FName CollectionName, ECollec
 
 bool FCollectionContainer::HasCollectionColors(TArray<FLinearColor>* OutColors) const
 {
-	FCollectionLock_RW Guard(Lock);
+	FCollectionScopeLock_RW Guard(Lock);
 	CollectionCache->UpdateCaches(Guard, ECollectionCacheFlags::Colors);
 	const TArray<FLinearColor>& CollectionColors = CollectionCache->GetCachedColors(Guard);
 	if (OutColors)
@@ -1818,7 +1953,7 @@ bool FCollectionContainer::GetCollectionColor(FName CollectionName, ECollectionS
 		return false;
 	}
 
-	FCollectionLock_Read Guard(Lock);
+	FCollectionScopeLock_Read Guard(Lock);
 	const FCollectionNameType CollectionKey(CollectionName, ShareType);
 	const TSharedRef<FCollection>* const CollectionRefPtr = AvailableCollections.Find(CollectionKey);
 	if (CollectionRefPtr)
@@ -1841,7 +1976,7 @@ bool FCollectionContainer::SetCollectionColor(FName CollectionName, ECollectionS
 {
 	const FCollectionNameType CollectionKey(CollectionName, ShareType);
 	{
-		FCollectionLock_Write Guard(Lock);
+		FCollectionScopeLock_Write Guard(Lock);
 
 		if (!ValidateWritable(Guard, ShareType, OutError))
 		{
@@ -1888,7 +2023,7 @@ bool FCollectionContainer::GetCollectionStorageMode(FName CollectionName, EColle
 		return false;
 	}
 
-	FCollectionLock_Read Guard(Lock);
+	FCollectionScopeLock_Read Guard(Lock);
 	const FCollectionNameType CollectionKey(CollectionName, ShareType);
 	const TSharedRef<FCollection>* const CollectionRefPtr = AvailableCollections.Find(CollectionKey);
 	if (CollectionRefPtr)
@@ -1919,7 +2054,7 @@ bool FCollectionContainer::IsObjectInCollection(const FSoftObjectPath& ObjectPat
 		return false;
 	}
 
-	FCollectionLock_RW Guard(Lock);
+	FCollectionScopeLock_RW Guard(Lock);
 	CollectionCache->UpdateCaches(Guard, ECollectionCacheFlags::RecursionWorker);
 	bool bFoundObject = false;
 
@@ -1940,12 +2075,12 @@ bool FCollectionContainer::IsObjectInCollection(const FSoftObjectPath& ObjectPat
 
 bool FCollectionContainer::IsValidParentCollection(FName CollectionName, ECollectionShareType::Type ShareType, FName ParentCollectionName,ECollectionShareType::Type ParentShareType, FText* OutError) const
 {
-	FCollectionLock_RW Guard(Lock);
+	FCollectionScopeLock_RW Guard(Lock);
 	CollectionCache->UpdateCaches(Guard, ECollectionCacheFlags::RecursionWorker);
 	return IsValidParentCollection_Locked(Guard, CollectionName, ShareType, ParentCollectionName, ParentShareType, OutError);
 }
 
-bool FCollectionContainer::IsValidParentCollection_Locked(FCollectionLock& InGuard, FName CollectionName, ECollectionShareType::Type ShareType, FName ParentCollectionName, ECollectionShareType::Type ParentShareType, FText* OutError) const 
+bool FCollectionContainer::IsValidParentCollection_Locked(FCollectionScopeLock& InGuard, FName CollectionName, ECollectionShareType::Type ShareType, FName ParentCollectionName, ECollectionShareType::Type ParentShareType, FText* OutError) const
 {
 	if (!ensure(ShareType < ECollectionShareType::CST_All) || (!ParentCollectionName.IsNone() && !ensure(ParentShareType < ECollectionShareType::CST_All)))
 	{
@@ -2014,11 +2149,11 @@ bool FCollectionContainer::IsValidParentCollection_Locked(FCollectionLock& InGua
 	return bValidParent;
 }
 
-void FCollectionContainer::UpdateCaches(ECollectionCacheFlags ToUpdate)
+bool FCollectionContainer::UpdateCaches(ECollectionCacheFlags ToUpdate)
 {
-	FCollectionLock_RW Guard(Lock);
+	FCollectionScopeLock_RW Guard(Lock);
 
-	CollectionCache->UpdateCaches(Guard, ToUpdate);
+	return CollectionCache->UpdateCaches(Guard, ToUpdate);
 }
 
 void FCollectionContainer::HandleFixupRedirectors(ICollectionRedirectorFollower& InRedirectorFollower)
@@ -2027,9 +2162,9 @@ void FCollectionContainer::HandleFixupRedirectors(ICollectionRedirectorFollower&
 	TArray<FSoftObjectPath> AddedObjects;
 	TArray<FSoftObjectPath> RemovedObjects;
 	{
-		FCollectionLock_Write Guard(Lock);
+		FCollectionScopeLock_Write Guard(Lock);
 
-		CollectionCache->UpdateCaches(Guard, ECollectionCacheFlags::Objects);
+		verifyf(CollectionCache->UpdateCaches(Guard, ECollectionCacheFlags::Objects), TEXT("UpdateCaches must be called within a write lock to guarantee subsequent usage to function as expected"));
 
 		const double LoadStartTime = FPlatformTime::Seconds();
 
@@ -2085,7 +2220,7 @@ bool FCollectionContainer::HandleRedirectorsDeleted(TConstArrayView<FSoftObjectP
 	bool bSavedAllCollections = true;
 	TArray<FCollectionNameType> UpdatedCollections;
 	{ 
-		FCollectionLock_Write Guard(Lock);
+		FCollectionScopeLock_Write Guard(Lock);
 		TSet<FCollectionNameType> CollectionsToSave;
 		FTextBuilder ErrorBuilder;
 
@@ -2148,7 +2283,7 @@ void FCollectionContainer::HandleObjectRenamed(const FSoftObjectPath& OldObjectP
 	TArray<FSoftObjectPath> AddedObjects;
 	TArray<FSoftObjectPath> RemovedObjects;
 	{
-		FCollectionLock_Write Guard(Lock);
+		FCollectionScopeLock_Write Guard(Lock);
 
 		ReplaceObjectInCollections(Guard, OldObjectPath, NewObjectPath, UpdatedCollections);
 
@@ -2173,9 +2308,9 @@ void FCollectionContainer::HandleObjectsDeleted(TConstArrayView<FSoftObjectPath>
 {
 	TArray<FCollectionNameType> UpdatedCollections;
 	{
-		FCollectionLock_Write Guard(Lock);
+		FCollectionScopeLock_Write Guard(Lock);
 
-		CollectionCache->UpdateCaches(Guard, ECollectionCacheFlags::Objects);
+		verifyf(CollectionCache->UpdateCaches(Guard, ECollectionCacheFlags::Objects), TEXT("UpdateCaches must be called within a write lock to guarantee subsequent usage to function as expected"));
 
 		for (const FSoftObjectPath& ObjectPath : ObjectPaths)
 		{
@@ -2197,7 +2332,7 @@ void FCollectionContainer::HandleObjectsDeleted(TConstArrayView<FSoftObjectPath>
 
 void FCollectionContainer::OnRemovedFromCollectionManager()
 {
-	FCollectionLock_Write Guard(Lock);
+	FCollectionScopeLock_Write Guard(Lock);
 
 	CollectionManager = nullptr;
 }
@@ -2216,7 +2351,7 @@ void FCollectionContainer::TickFileCache()
 	TArray<TTuple<ECollectionFileAction, FCollectionNameType>> Events;
 	{
 		// Acquire write lock immediately so we don't need to deal with state change during promotion
-		FCollectionLock_Write Guard(Lock);
+		FCollectionScopeLock_Write Guard(Lock);
 
 		// Process changes that have happened outside of the collection container
 		for (int32 CacheIdx = 0; CacheIdx < ECollectionShareType::CST_All; ++CacheIdx)
@@ -2381,7 +2516,7 @@ void FCollectionContainer::LoadCollections()
 					TSharedRef<FCollection> NewCollection = MakeShareable(new FCollection(Filename, bUseSCC, ECollectionStorageMode::Static));
 					if (NewCollection->Load(LoadErrorText))
 					{
-						FCollectionLock_Write Guard(Lock);
+						FCollectionScopeLock_Write Guard(Lock);
 						AddCollection(Guard, NewCollection, ShareType);
 					}
 					else
@@ -2396,7 +2531,7 @@ void FCollectionContainer::LoadCollections()
 	);
 
 	// AddCollection is assumed to be adding an empty collection, so also notify that collection cache that the collection has "changed" since loaded collections may not always be empty
-	FCollectionLock_Write Guard(Lock);
+	FCollectionScopeLock_Write Guard(Lock);
 	CollectionCache->HandleCollectionChanged(Guard);
 
 	UE_LOG(LogCollectionManager, Log, TEXT( "Loaded %d collections in %0.6f seconds" ), AvailableCollections.Num() - PrevNumCollections, FPlatformTime::Seconds() - LoadStartTime);
@@ -2435,13 +2570,13 @@ uint8 FCollectionContainer::GetReadOnlyMask(ECollectionShareType::Type ShareType
 	return Mask;
 }
 
-bool FCollectionContainer::IsReadOnly(FCollectionLock& InGuard, ECollectionShareType::Type ShareType) const
+bool FCollectionContainer::IsReadOnly(FCollectionScopeLock& InGuard, ECollectionShareType::Type ShareType) const
 {
 	const uint8 ReadOnlyMask = GetReadOnlyMask(ShareType);
 	return (ReadOnlyFlags & ReadOnlyMask) == ReadOnlyMask;
 }
 
-bool FCollectionContainer::ValidateWritable(FCollectionLock& InGuard, ECollectionShareType::Type ShareType, FText* OutError) const
+bool FCollectionContainer::ValidateWritable(FCollectionScopeLock& InGuard, ECollectionShareType::Type ShareType, FText* OutError) const
 {
 	if (!ensure(ShareType < ECollectionShareType::CST_All))
 	{
@@ -2473,7 +2608,7 @@ bool FCollectionContainer::ValidateWritable(FCollectionLock& InGuard, ECollectio
 	return true;
 }
 
-bool FCollectionContainer::AddCollection(FCollectionLock_Write& Guard, const TSharedRef<FCollection>& CollectionRef, ECollectionShareType::Type ShareType)
+bool FCollectionContainer::AddCollection(FCollectionScopeLock_Write& Guard, const TSharedRef<FCollection>& CollectionRef, ECollectionShareType::Type ShareType)
 {
 	if (!ensure(ShareType < ECollectionShareType::CST_All))
 	{
@@ -2493,7 +2628,7 @@ bool FCollectionContainer::AddCollection(FCollectionLock_Write& Guard, const TSh
 	return true;
 }
 
-bool FCollectionContainer::RemoveCollection(FCollectionLock_Write& Guard, const TSharedRef<FCollection>& CollectionRef, ECollectionShareType::Type ShareType)
+bool FCollectionContainer::RemoveCollection(FCollectionScopeLock_Write& Guard, const TSharedRef<FCollection>& CollectionRef, ECollectionShareType::Type ShareType)
 {
 	if (!ensure(ShareType < ECollectionShareType::CST_All))
 	{
@@ -2511,7 +2646,7 @@ bool FCollectionContainer::RemoveCollection(FCollectionLock_Write& Guard, const
 	return false;
 }
 
-void FCollectionContainer::RemoveObjectFromCollections(FCollectionLock_Write& Guard, const FSoftObjectPath& ObjectPath, TArray<FCollectionNameType>& OutUpdatedCollections)
+void FCollectionContainer::RemoveObjectFromCollections(FCollectionScopeLock_Write& Guard, const FSoftObjectPath& ObjectPath, TArray<FCollectionNameType>& OutUpdatedCollections)
 {
 	const TMap<FSoftObjectPath, TArray<FObjectCollectionInfo>>& CachedObjects = CollectionCache->GetCachedObjects(Guard);
 
@@ -2538,9 +2673,9 @@ void FCollectionContainer::RemoveObjectFromCollections(FCollectionLock_Write& Gu
 	}
 }
 
-void FCollectionContainer::ReplaceObjectInCollections(FCollectionLock_Write& InGuard, const FSoftObjectPath& OldObjectPath, const FSoftObjectPath& NewObjectPath, TArray<FCollectionNameType>& OutUpdatedCollections)
+void FCollectionContainer::ReplaceObjectInCollections(FCollectionScopeLock_Write& InGuard, const FSoftObjectPath& OldObjectPath, const FSoftObjectPath& NewObjectPath, TArray<FCollectionNameType>& OutUpdatedCollections)
 {
-	CollectionCache->UpdateCaches(InGuard, ECollectionCacheFlags::Objects);
+	verifyf(CollectionCache->UpdateCaches(InGuard, ECollectionCacheFlags::Objects), TEXT("UpdateCaches must be called within a write lock to guarantee subsequent usage to function as expected"));
 	const TMap<FSoftObjectPath, TArray<FObjectCollectionInfo>>& CachedObjects = CollectionCache->GetCachedObjects(InGuard);
 
 	const TArray<FObjectCollectionInfo>* OldObjectCollectionInfosPtr = CachedObjects.Find(OldObjectPath);
@@ -2567,7 +2702,7 @@ void FCollectionContainer::ReplaceObjectInCollections(FCollectionLock_Write& InG
 	}
 }
 
-bool FCollectionContainer::InternalSaveCollection(FCollectionLock_Write& InGuard, const TSharedRef<FCollection>& CollectionRef, FText* OutError, bool bForceCommitToRevisionControl)
+bool FCollectionContainer::InternalSaveCollection(FCollectionScopeLock_Write& InGuard, const TSharedRef<FCollection>& CollectionRef, FText* OutError, bool bForceCommitToRevisionControl)
 {
 	if (!ensure(CollectionManager != nullptr))
 	{
diff --git a/Engine/Source/Developer/CollectionManager/Private/CollectionContainer.h b/Engine/Source/Developer/CollectionManager/Private/CollectionContainer.h
--- a/Engine/Source/Developer/CollectionManager/Private/CollectionContainer.h
+++ b/Engine/Source/Developer/CollectionManager/Private/CollectionContainer.h
@@ -16,16 +16,58 @@ enum class ECollectionCacheFlags;
 class FCollectionContainerCache;
 class FCollectionManager;
 
+UE_INTERNAL class FCollectionRecursiveRwLock
+{
+	struct FThreadLockDepths
+	{
+		FThreadLockDepths(void* TlsSlotValue);
+
+		void* GetTlsSlotValue();
+
+#if PLATFORM_64BITS
+		static_assert(sizeof(UPTRINT) == sizeof(uint64), "Expected pointer size to be 64 bits");
+
+		using HalfUPtrInt = uint32;
+#else
+		static_assert(sizeof(UPTRINT) == sizeof(uint32), "Expected pointer size to be 32 bits");
+
+		using HalfUPtrInt = uint16;
+#endif
+
+		HalfUPtrInt ThreadReadDepth = 0;
+		HalfUPtrInt ThreadWriteDepth = 0;
+	};
+
+public:
+	FCollectionRecursiveRwLock();
+	~FCollectionRecursiveRwLock();
+
+	void ReadLock();
+	void WriteLock();
+
+	void ReadUnlock();
+	void WriteUnlock();
+
+	// Promote the lock from read to write, possibly being interrupted by another writer in between.
+	// Returns false if the lock cannot be promoted, which can happen if the thread already holds a read
+	// lock then reenters and tries to promote to a write lock.
+	[[nodiscard]] bool PromoteInterruptible();
+
+private:
+	FRWLock RwLock;
+	uint32 TlsSlot;
+};
+
 // Objects wrapping locks to read, write, or begin-reading-then-write (for cache updates) internal state. 
 // Used as internal function parameters to show what lock type must be held to perform the operation and prevent 
 // recursive lock acquisition
 // Functions taking Lock/Lock_Read need to be able to read data but not update caches
-class FCollectionLock;
-class FCollectionLock_Read;
+class FCollectionScopeLock;
+class FCollectionScopeLock_Read;
 // Functions taking Lock_RW may need to promote the lock to a write state to update caches
-class FCollectionLock_RW;
+class FCollectionScopeLock_RW;
 // Functions taking Lock_Write have exclusive access and can update collections as well as update caches
-class FCollectionLock_Write;
+class FCollectionScopeLock_Write;
 
 class FCollectionContainer final : public ICollectionContainer
 {
@@ -84,7 +126,11 @@ public:
 	virtual bool IsObjectInCollection(const FSoftObjectPath& ObjectPath, FName CollectionName, ECollectionShareType::Type ShareType, ECollectionRecursionFlags::Flags RecursionMode = ECollectionRecursionFlags::Self, FText* OutError = nullptr) const override;
 	virtual bool IsValidParentCollection(FName CollectionName, ECollectionShareType::Type ShareType, FName ParentCollectionName, ECollectionShareType::Type ParentShareType, FText* OutError) const override;
 
-	void UpdateCaches(ECollectionCacheFlags ToUpdate);
+	/*
+	* Returns false if the collection container's internal lock cannot be promoted preventing the cache from updating,
+	* which can happen if the thread already holds a read lock then reenters and tries to promote to a write lock.
+	*/
+	bool UpdateCaches(ECollectionCacheFlags ToUpdate);
 
 	void HandleFixupRedirectors(ICollectionRedirectorFollower& InRedirectorFollower);
 	bool HandleRedirectorsDeleted(TConstArrayView<FSoftObjectPath> ObjectPaths, FText* OutError);
@@ -139,23 +185,23 @@ private:
 	/** Returns the read-only bit mask for the specified share type. */
 	static uint8 GetReadOnlyMask(ECollectionShareType::Type ShareType);
 
-	bool IsReadOnly(FCollectionLock& InGuard, ECollectionShareType::Type ShareType) const;
+	bool IsReadOnly(FCollectionScopeLock& InGuard, ECollectionShareType::Type ShareType) const;
 
 	/** Returns whether the collection container is in a valid state for writing */
-	bool ValidateWritable(FCollectionLock& InGuard, ECollectionShareType::Type ShareType, FText* OutError) const;
+	bool ValidateWritable(FCollectionScopeLock& InGuard, ECollectionShareType::Type ShareType, FText* OutError) const;
 
 	/** Adds a collection to the lookup maps */
-	bool AddCollection(FCollectionLock_Write& InGuard, const TSharedRef<FCollection>& CollectionRef, ECollectionShareType::Type ShareType);
+	bool AddCollection(FCollectionScopeLock_Write& InGuard, const TSharedRef<FCollection>& CollectionRef, ECollectionShareType::Type ShareType);
 
 	/** Removes a collection from the lookup maps */
-	bool RemoveCollection(FCollectionLock_Write& InGuard, const TSharedRef<FCollection>& CollectionRef, ECollectionShareType::Type ShareType);
+	bool RemoveCollection(FCollectionScopeLock_Write& InGuard, const TSharedRef<FCollection>& CollectionRef, ECollectionShareType::Type ShareType);
 
 	/** Removes an object from any collections that contain it */
-	void RemoveObjectFromCollections(FCollectionLock_Write& InGuard, const FSoftObjectPath& ObjectPath, TArray<FCollectionNameType>& OutUpdatedCollections);
+	void RemoveObjectFromCollections(FCollectionScopeLock_Write& InGuard, const FSoftObjectPath& ObjectPath, TArray<FCollectionNameType>& OutUpdatedCollections);
 
 	/** Replaces an object with another in any collections that contain it */
 	void ReplaceObjectInCollections(
-		FCollectionLock_Write& InGuard, const FSoftObjectPath& OldObjectPath, const FSoftObjectPath& NewObjectPath,
+		FCollectionScopeLock_Write& InGuard, const FSoftObjectPath& OldObjectPath, const FSoftObjectPath& NewObjectPath,
 		TArray<FCollectionNameType>& OutUpdatedCollections);
 
 	/** Internal common functionality for saving a collection
@@ -163,26 +209,26 @@ private:
 	 * bForceCommitToRevisionControl will ensure that it is committed after save.  If this is false, then the collection
 	 * will be left as a modified file which can be advantageous for slow source control servers.
 	 */
-	bool InternalSaveCollection(FCollectionLock_Write&, const TSharedRef<FCollection>& CollectionRef, FText* OutError, bool bForceCommitToRevisionControl);
+	bool InternalSaveCollection(FCollectionScopeLock_Write&, const TSharedRef<FCollection>& CollectionRef, FText* OutError, bool bForceCommitToRevisionControl);
 
 	/* 
 	 * Internal version of IsValidParentCollection to avoid taking lock recursively.
 	 * Cache must be updated for recursion before calling.
 	 */
-	bool IsValidParentCollection_Locked(FCollectionLock& InGuard, FName CollectionName, ECollectionShareType::Type ShareType, FName ParentCollectionName, ECollectionShareType::Type ParentShareType, FText* OutError) const;
+	bool IsValidParentCollection_Locked(FCollectionScopeLock& InGuard, FName CollectionName, ECollectionShareType::Type ShareType, FName ParentCollectionName, ECollectionShareType::Type ParentShareType, FText* OutError) const;
 
 	/** 
 	 * Check if the given collection exists.
 	 * Using the public API function risks acquiring the lock recursively.
 	 */
-	bool CollectionExists_Locked(FCollectionLock& InGuard, FName CollectionName, ECollectionShareType::Type ShareType) const;
+	bool CollectionExists_Locked(FCollectionScopeLock& InGuard, FName CollectionName, ECollectionShareType::Type ShareType) const;
 
 private:
 	/** The extension used for collection files */
 	static FStringView CollectionExtension;
 
 	/** Required for updating caches as well as write operations to collections */
-	mutable FRWLock Lock;
+	mutable FCollectionRecursiveRwLock Lock;
 
 	/** The collection managed that is managing this collection container.
 	  * Null if this collection container has been removed from the collection manager.
