Collection.cpp: Optimize loading path for utf8 encoded files, fallback to FString/TChar conversion for other file encodings.
- 0.8s to 0.55s at boot for some projects.

#rnx

diff --git a/Engine/Source/Developer/CollectionManager/Private/Collection.cpp b/Engine/Source/Developer/CollectionManager/Private/Collection.cpp
--- a/Engine/Source/Developer/CollectionManager/Private/Collection.cpp
+++ b/Engine/Source/Developer/CollectionManager/Private/Collection.cpp
@@ -77,20 +77,64 @@ TSharedRef<FCollection> FCollection::Clone(const FString& InFilename, bool InUse
 	return NewCollection;
 }
 
+static bool ParseUtf8Bom(FMutableMemoryView& Buffer)
+{
+	const uint64 Size = Buffer.GetSize();
+	const uint8* const Bytes = static_cast<const uint8*>(Buffer.GetData());
+	if (!(Size & 1) && Size >= 2 && Bytes[0] == 0xff && Bytes[1] == 0xfe)
+	{
+		// Unicode Intel byte order BOM, leave it to be consumed by FFileHelper::BufferToString
+		return false;
+	}
+	if (!(Size & 1) && Size >= 2 && Bytes[0] == 0xfe && Bytes[1] == 0xff)
+	{
+		// Unicode non-Intel byte order BOM, leave it to be consumed by FFileHelper::BufferToString
+		return false;
+	}
+	if (Size >= 3 && Bytes[0] == 0xef && Bytes[1] == 0xbb && Bytes[2] == 0xbf)
+	{
+		// Utf8 BOM, consume it and treat buffer as Utf8 without BOM
+		Buffer += 3;
+	}
+	// Assume Utf8 without BOM
+	return true;
+}
+
 bool FCollection::Load(FText& OutError)
 {
 	TRACE_CPUPROFILER_EVENT_SCOPE(FCollection::Load)
 	Empty();
 
-	FString FullFileContentsString;
-	if (!FFileHelper::LoadFileToString(FullFileContentsString, *SourceFilename))
+	TArray<uint8> Buffer;
+	if (!FFileHelper::LoadFileToArray(Buffer, *SourceFilename))
 	{
 		OutError = FText::Format(LOCTEXT("LoadError_FailedToLoadFile", "Failed to load the collection '{0}' from disk."), FText::FromString(SourceFilename));
 		return false;
 	}
 
-	TArray<FStringView> FileContents;
-	UE::String::ParseLines(FullFileContentsString, [&FileContents](const FStringView& Line) { FileContents.Add(Line); });
+	FMutableMemoryView BufferView(Buffer.GetData(), Buffer.Num());
+	if (ParseUtf8Bom(BufferView))
+	{
+		FUtf8StringView FullFileContentsString(
+			reinterpret_cast<const UTF8CHAR*>(BufferView.GetData()), IntCastChecked<int32>(BufferView.GetSize()));
+		return ParseCollection(FullFileContentsString, OutError);
+	}
+	else
+	{
+		FString FullFileContentsString;
+		FFileHelper::BufferToString(FullFileContentsString, Buffer.GetData(), Buffer.Num());
+		Buffer.Reset();
+		return ParseCollection(FStringView(FullFileContentsString), OutError);
+	}
+}
+
+template <typename CharType>
+bool FCollection::ParseCollection(TStringView<CharType> FullFileContentsString, FText& OutError)
+{
+	// Reserve an initial estimate and parse all lines
+	TArray<TStringView<CharType>> FileContents;
+	FileContents.Reserve(FullFileContentsString.Len() / 64);
+	UE::String::ParseLines(FullFileContentsString, [&FileContents](const TStringView<CharType>& Line) { FileContents.Add(Line); });
 
 	if (FileContents.Num() == 0)
 	{
@@ -104,7 +148,7 @@ bool FCollection::Load(FText& OutError)
 	int32 LineIndex = 0;
 	for (int32 Num = FileContents.Num(); LineIndex < Num; ++LineIndex)
 	{
-		FStringView Line(FileContents[LineIndex]);
+		TStringView<CharType> Line(FileContents[LineIndex]);
 		Line.TrimStartAndEndInline();
 
 		if (Line.Len() == 0)
@@ -115,7 +159,7 @@ bool FCollection::Load(FText& OutError)
 		}
 
 		int32 Offset;
-		if (Line.FindChar(TEXT(':'), Offset))
+		if (Line.FindChar(CharType(':'), Offset))
 		{
 			FString Key(Line.Left(Offset));
 			FString Value(Line.Right(Line.Len() - Offset - 1));
@@ -134,42 +178,43 @@ bool FCollection::Load(FText& OutError)
 	// Now load the content if the header load was successful
 	if (StorageMode == ECollectionStorageMode::Static)
 	{
-		const int32 NamesNum = FileContents.Num() - LineIndex;
-
-		struct FItem
+		if (const int32 NamesNum = FileContents.Num() - LineIndex; NamesNum > 0)
 		{
-			FSoftObjectPath Path;
-			int32 Hash;
-		};
-
-		TArray<FItem> Paths;
-		Paths.SetNum(NamesNum);
-
-		typedef decltype(ObjectSet)::KeyFuncsType KeyFuncsType;
-		// Name hashing to register new FName takes time
-		// Process as much as possible in multiple threads
-		ParallelFor(
-			NamesNum,
-			[&FileContents, &LineIndex, &Paths](int32 LocalLineIndex)
+			struct FItem
 			{
-				FStringView Line(FileContents[LineIndex + LocalLineIndex]);
-				Paths[LocalLineIndex].Path = FSoftObjectPath(Line.TrimStartAndEnd());
+				FSoftObjectPath Path;
+				int32 Hash;
+			};
+
+			TArray<FItem> Paths;
+			Paths.SetNum(NamesNum);
+
+			typedef decltype(ObjectSet)::KeyFuncsType KeyFuncsType;
+			// Name hashing to register new FName takes time
+			// Process as much as possible in multiple threads
+			TStringView<CharType>* Lines = &FileContents[LineIndex];
+			ParallelFor(
+				NamesNum,
+				[&Lines, &Paths](int32 LocalLineIndex)
+			{
+				Paths[LocalLineIndex].Path = FSoftObjectPath::ConstructFromStringPath(Lines[LocalLineIndex].TrimStartAndEnd());
 				Paths[LocalLineIndex].Hash = KeyFuncsType::GetKeyHash(KeyFuncsType::GetSetKey(Paths[LocalLineIndex].Path));
 			},
-			// Do not pay for scheduling cost if number of items is too low
-			NamesNum < 1000 ? EParallelForFlags::ForceSingleThread : EParallelForFlags::None
-		);
+				// Do not pay for scheduling cost if number of items is too low
+				NamesNum < 1000 ? EParallelForFlags::ForceSingleThread : EParallelForFlags::None
+			);
 
-		ObjectSet.Reserve(ObjectSet.Num() + Paths.Num());
-		// Static collection, a flat list of asset paths
-		for (FItem& Item : Paths)
-		{
-			if (Item.Path.IsNull())
+			ObjectSet.Reserve(ObjectSet.Num() + Paths.Num());
+			// Static collection, a flat list of asset paths
+			for (FItem& Item : Paths)
 			{
-				continue;
-			}
+				if (Item.Path.IsNull())
+				{
+					continue;
+				}
 
-			ObjectSet.EmplaceByHash(Item.Hash, MoveTemp(Item.Path));
+				ObjectSet.EmplaceByHash(Item.Hash, MoveTemp(Item.Path));
+			}
 		}
 	}
 	else
diff --git a/Engine/Source/Developer/CollectionManager/Private/Collection.h b/Engine/Source/Developer/CollectionManager/Private/Collection.h
--- a/Engine/Source/Developer/CollectionManager/Private/Collection.h
+++ b/Engine/Source/Developer/CollectionManager/Private/Collection.h
@@ -149,6 +149,9 @@ private:
 	  */
 	bool LoadHeaderPairs(const TMap<FString,FString>& InHeaderPairs);
 
+	template <typename CharType>
+	bool ParseCollection(TStringView<CharType> FullFileContentsString, FText& OutError);
+
 	/** Merges the assets from the specified collection with this collection */
 	bool MergeWithCollection(const FCollection& Other);
 	/** Gets the object differences between object set A (base) and B (new) */
diff --git a/Engine/Source/Runtime/Core/Private/String/ParseLines.cpp b/Engine/Source/Runtime/Core/Private/String/ParseLines.cpp
--- a/Engine/Source/Runtime/Core/Private/String/ParseLines.cpp
+++ b/Engine/Source/Runtime/Core/Private/String/ParseLines.cpp
@@ -7,29 +7,30 @@
 namespace UE::String
 {
 
-void ParseLines(
-	const FStringView View,
-	const TFunctionRef<void (FStringView)> Visitor,
+template <typename CharType>
+void ParseLinesImpl(
+	const TStringView<CharType> View,
+	const TFunctionRef<void (TStringView<CharType>)> Visitor,
 	const EParseLinesOptions Options)
 {
-	const TCHAR* ViewIt = View.GetData();
-	const TCHAR* ViewEnd = ViewIt + View.Len();
+	const CharType* ViewIt = View.GetData();
+	const CharType* ViewEnd = ViewIt + View.Len();
 	do
 	{
-		const TCHAR* LineStart = ViewIt;
-		const TCHAR* LineEnd = ViewEnd;
+		const CharType* LineStart = ViewIt;
+		const CharType* LineEnd = ViewEnd;
 		for (; ViewIt != ViewEnd; ++ViewIt)
 		{
-			const TCHAR CurrentChar = *ViewIt;
-			if (CurrentChar == TEXT('\n'))
+			const CharType CurrentChar = *ViewIt;
+			if (CurrentChar == CharType('\n'))
 			{
 				LineEnd = ViewIt++;
 				break;
 			}
-			if (CurrentChar == TEXT('\r'))
+			if (CurrentChar == CharType('\r'))
 			{
 				LineEnd = ViewIt++;
-				if (ViewIt != ViewEnd && *ViewIt == TEXT('\n'))
+				if (ViewIt != ViewEnd && *ViewIt == CharType('\n'))
 				{
 					++ViewIt;
 				}
@@ -37,7 +38,7 @@ void ParseLines(
 			}
 		}
 
-		FStringView Line(LineStart, UE_PTRDIFF_TO_INT32(LineEnd - LineStart));
+		TStringView<CharType> Line(LineStart, UE_PTRDIFF_TO_INT32(LineEnd - LineStart));
 		if (EnumHasAnyFlags(Options, EParseLinesOptions::Trim))
 		{
 			Line = Line.TrimStartAndEnd();
@@ -50,4 +51,19 @@ void ParseLines(
 	while (ViewIt != ViewEnd);
 }
 
+void ParseLines(
+	FStringView View,
+	TFunctionRef<void (FStringView)> Visitor,
+	EParseLinesOptions Options)
+{
+	return ParseLinesImpl(View, Visitor, Options);
+}
+
+void ParseLines(
+	FUtf8StringView View,
+	TFunctionRef<void (FUtf8StringView)> Visitor,
+	EParseLinesOptions Options)
+{
+	return ParseLinesImpl(View, Visitor, Options);
+}
 } // UE::String
diff --git a/Engine/Source/Runtime/Core/Public/String/ParseLines.h b/Engine/Source/Runtime/Core/Public/String/ParseLines.h
--- a/Engine/Source/Runtime/Core/Public/String/ParseLines.h
+++ b/Engine/Source/Runtime/Core/Public/String/ParseLines.h
@@ -38,6 +38,10 @@ CORE_API void ParseLines(
 	FStringView View,
 	TFunctionRef<void (FStringView)> Visitor,
 	EParseLinesOptions Options = EParseLinesOptions::None);
+CORE_API void ParseLines(
+	FUtf8StringView View,
+	TFunctionRef<void (FUtf8StringView)> Visitor,
+	EParseLinesOptions Options = EParseLinesOptions::None);
 
 /**
  * Parse lines in the input string as terminated by any of CRLF, CR, LF.
@@ -49,13 +53,13 @@ CORE_API void ParseLines(
  * @param Output    The output to add parsed lines to by calling Output.Add(FStringView).
  * @param Options   Flags to modify the default behavior.
  */
-template <typename OutputType>
+template <typename OutputType, typename StringViewType>
 inline void ParseLines(
-	const FStringView View,
+	const StringViewType View,
 	OutputType& Output,
 	const EParseLinesOptions Options = EParseLinesOptions::None)
 {
-	ParseLines(View, [&Output](FStringView Line) { Output.Add(Line); }, Options);
+	ParseLines(View, [&Output](StringViewType Line) { Output.Add(Line); }, Options);
 }
 
 } // UE::String
