Return any errors found when calling the single asset AddToCollection overload.

diff --git a/Engine/Source/Developer/CollectionManager/Private/CollectionContainer.cpp b/Engine/Source/Developer/CollectionManager/Private/CollectionContainer.cpp
--- a/Engine/Source/Developer/CollectionManager/Private/CollectionContainer.cpp
+++ b/Engine/Source/Developer/CollectionManager/Private/CollectionContainer.cpp
@@ -1467,7 +1467,7 @@ bool FCollectionContainer::DestroyCollection(FName CollectionName, ECollectionSh
 
 bool FCollectionContainer::AddToCollection(FName CollectionName, ECollectionShareType::Type ShareType, const FSoftObjectPath& ObjectPath, FText* OutError)
 {
-	return AddToCollection(CollectionName, ShareType, MakeArrayView(&ObjectPath, 1));
+	return AddToCollection(CollectionName, ShareType, MakeArrayView(&ObjectPath, 1), nullptr, OutError);
 }
 
 bool FCollectionContainer::AddToCollection(FName CollectionName, ECollectionShareType::Type ShareType, TConstArrayView<FSoftObjectPath> ObjectPaths, int32* OutNumAdded, FText* OutError)
