Updated ICollectionManager UE_DEPRECATED messages to make it clearer how to update existing code to the new API.

diff --git a/Engine/Source/Developer/CollectionManager/Public/ICollectionManager.h b/Engine/Source/Developer/CollectionManager/Public/ICollectionManager.h
--- a/Engine/Source/Developer/CollectionManager/Public/ICollectionManager.h
+++ b/Engine/Source/Developer/CollectionManager/Public/ICollectionManager.h
@@ -46,43 +46,43 @@ public:
 	virtual bool TryParseCollectionPath(const FString& CollectionPath, TSharedPtr<ICollectionContainer>* OutCollectionContainer, FName* OutCollectionName, ECollectionShareType::Type* OutShareType) const = 0;
 
 	/** Returns whether or not the collection manager contains any collections */
-	UE_DEPRECATED(5.6, "ICollectionContainer has been introduced as a new top level object for housing collections, call GetProjectCollectionContainer to disambiguate which collection container to use.")
+	UE_DEPRECATED(5.6, "ICollectionContainer has been introduced as a new top level object for housing collections, substitute wtih GetProjectCollectionContainer()->HasCollections() to disambiguate which collection container to use.")
 	virtual bool HasCollections() const = 0;
 
 	/** Returns the list of collections */
-	UE_DEPRECATED(5.6, "ICollectionContainer has been introduced as a new top level object for housing collections, call GetProjectCollectionContainer to disambiguate which collection container to use.")
+	UE_DEPRECATED(5.6, "ICollectionContainer has been introduced as a new top level object for housing collections, substitute wtih GetProjectCollectionContainer()->GetCollections(...) to disambiguate which collection container to use.")
 	virtual void GetCollections(TArray<FCollectionNameType>& OutCollections) const = 0;
 
 	/** Returns the list of collections with the given name */
-	UE_DEPRECATED(5.6, "ICollectionContainer has been introduced as a new top level object for housing collections, call GetProjectCollectionContainer to disambiguate which collection container to use.")
+	UE_DEPRECATED(5.6, "ICollectionContainer has been introduced as a new top level object for housing collections, substitute wtih GetProjectCollectionContainer()->GetCollections(...) to disambiguate which collection container to use.")
 	virtual void GetCollections(FName CollectionName, TArray<FCollectionNameType>& OutCollections) const = 0;
 
 	/** Returns the list of collection names of the specified share type */
-	UE_DEPRECATED(5.6, "ICollectionContainer has been introduced as a new top level object for housing collections, call GetProjectCollectionContainer to disambiguate which collection container to use.")
+	UE_DEPRECATED(5.6, "ICollectionContainer has been introduced as a new top level object for housing collections, substitute wtih GetProjectCollectionContainer()->GetCollectionNames(...) to disambiguate which collection container to use.")
 	virtual void GetCollectionNames(ECollectionShareType::Type ShareType, TArray<FName>& CollectionNames) const = 0;
 
 	/** Returns the list of root-level collections */
-	UE_DEPRECATED(5.6, "ICollectionContainer has been introduced as a new top level object for housing collections, call GetProjectCollectionContainer to disambiguate which collection container to use.")
+	UE_DEPRECATED(5.6, "ICollectionContainer has been introduced as a new top level object for housing collections, substitute wtih GetProjectCollectionContainer()->GetRootCollections(...) to disambiguate which collection container to use.")
 	virtual void GetRootCollections(TArray<FCollectionNameType>& OutCollections) const = 0;
 
 	/** Returns the list of root-level collection names of the specified share type */
-	UE_DEPRECATED(5.6, "ICollectionContainer has been introduced as a new top level object for housing collections, call GetProjectCollectionContainer to disambiguate which collection container to use.")
+	UE_DEPRECATED(5.6, "ICollectionContainer has been introduced as a new top level object for housing collections, substitute wtih GetProjectCollectionContainer()->GetRootCollectionNames(...) to disambiguate which collection container to use.")
 	virtual void GetRootCollectionNames(ECollectionShareType::Type ShareType, TArray<FName>& CollectionNames) const = 0;
 
 	/** Returns the list of child collections of the given collection */
-	UE_DEPRECATED(5.6, "ICollectionContainer has been introduced as a new top level object for housing collections, call GetProjectCollectionContainer to disambiguate which collection container to use.")
+	UE_DEPRECATED(5.6, "ICollectionContainer has been introduced as a new top level object for housing collections, substitute wtih GetProjectCollectionContainer()->GetChildCollections(...) to disambiguate which collection container to use.")
 	virtual void GetChildCollections(FName CollectionName, ECollectionShareType::Type ShareType, TArray<FCollectionNameType>& OutCollections) const = 0;
 
 	/** Returns the list of child collections of the given collection that are also of the specified share type */
-	UE_DEPRECATED(5.6, "ICollectionContainer has been introduced as a new top level object for housing collections, call GetProjectCollectionContainer to disambiguate which collection container to use.")
+	UE_DEPRECATED(5.6, "ICollectionContainer has been introduced as a new top level object for housing collections, substitute wtih GetProjectCollectionContainer()->GetChildCollectionNames(...) to disambiguate which collection container to use.")
 	virtual void GetChildCollectionNames(FName CollectionName, ECollectionShareType::Type ShareType, ECollectionShareType::Type ChildShareType, TArray<FName>& CollectionNames) const = 0;
 
 	/** Returns the parent collection of the given collection, or an unset value if there is no parent set */
-	UE_DEPRECATED(5.6, "ICollectionContainer has been introduced as a new top level object for housing collections, call GetProjectCollectionContainer to disambiguate which collection container to use.")
+	UE_DEPRECATED(5.6, "ICollectionContainer has been introduced as a new top level object for housing collections, substitute wtih GetProjectCollectionContainer()->GetParentCollection(...) to disambiguate which collection container to use.")
 	virtual TOptional<FCollectionNameType> GetParentCollection(FName CollectionName, ECollectionShareType::Type ShareType) const = 0;
 
 	/** Returns true if the collection exists */
-	UE_DEPRECATED(5.6, "ICollectionContainer has been introduced as a new top level object for housing collections, call GetProjectCollectionContainer to disambiguate which collection container to use.")
+	UE_DEPRECATED(5.6, "ICollectionContainer has been introduced as a new top level object for housing collections, substitute wtih GetProjectCollectionContainer()->CollectionExists(...) to disambiguate which collection container to use.")
 	virtual bool CollectionExists(FName CollectionName, ECollectionShareType::Type ShareType) const = 0;
 
 	/**
@@ -90,7 +90,7 @@ public:
 	 * 
 	 * @return true if any objects were found
 	 */
-	UE_DEPRECATED(5.6, "ICollectionContainer has been introduced as a new top level object for housing collections, call GetProjectCollectionContainer to disambiguate which collection container to use.")
+	UE_DEPRECATED(5.6, "ICollectionContainer has been introduced as a new top level object for housing collections, substitute wtih GetProjectCollectionContainer()->GetAssetsInCollection(...) to disambiguate which collection container to use.")
 	virtual bool GetAssetsInCollection(FName CollectionName, ECollectionShareType::Type ShareType, TArray<FSoftObjectPath>& OutAssetPaths, ECollectionRecursionFlags::Flags RecursionMode = ECollectionRecursionFlags::Self) const = 0;
 
 	/**
@@ -98,7 +98,7 @@ public:
 	 * 
 	 * @return true if any objects were found
 	 */
-	UE_DEPRECATED(5.6, "ICollectionContainer has been introduced as a new top level object for housing collections, call GetProjectCollectionContainer to disambiguate which collection container to use.")
+	UE_DEPRECATED(5.6, "ICollectionContainer has been introduced as a new top level object for housing collections, substitute wtih GetProjectCollectionContainer()->GetClassesInCollection(...) to disambiguate which collection container to use.")
 	virtual bool GetClassesInCollection(FName CollectionName, ECollectionShareType::Type ShareType, TArray<FTopLevelAssetPath>& ClassPaths, ECollectionRecursionFlags::Flags RecursionMode = ECollectionRecursionFlags::Self) const = 0;
 
 	/**
@@ -106,23 +106,23 @@ public:
 	 * 
 	 * @return true if any objects were found
 	 */
-	UE_DEPRECATED(5.6, "ICollectionContainer has been introduced as a new top level object for housing collections, call GetProjectCollectionContainer to disambiguate which collection container to use.")
+	UE_DEPRECATED(5.6, "ICollectionContainer has been introduced as a new top level object for housing collections, substitute wtih GetProjectCollectionContainer()->GetObjectsInCollection(...) to disambiguate which collection container to use.")
 	virtual bool GetObjectsInCollection(FName CollectionName, ECollectionShareType::Type ShareType, TArray<FSoftObjectPath>& ObjectPaths, ECollectionRecursionFlags::Flags RecursionMode = ECollectionRecursionFlags::Self) const = 0;
 
 	/** Returns a list of collections in which the specified object exists of the specified share type */
-	UE_DEPRECATED(5.6, "ICollectionContainer has been introduced as a new top level object for housing collections, call GetProjectCollectionContainer to disambiguate which collection container to use.")
+	UE_DEPRECATED(5.6, "ICollectionContainer has been introduced as a new top level object for housing collections, substitute wtih GetProjectCollectionContainer()->GetCollectionsContainingObject(...) to disambiguate which collection container to use.")
 	virtual void GetCollectionsContainingObject(const FSoftObjectPath& ObjectPath, ECollectionShareType::Type ShareType, TArray<FName>& OutCollectionNames, ECollectionRecursionFlags::Flags RecursionMode = ECollectionRecursionFlags::Self) const = 0;
 
 	/** Returns a list of collections in which the specified object exists */
-	UE_DEPRECATED(5.6, "ICollectionContainer has been introduced as a new top level object for housing collections, call GetProjectCollectionContainer to disambiguate which collection container to use.")
+	UE_DEPRECATED(5.6, "ICollectionContainer has been introduced as a new top level object for housing collections, substitute wtih GetProjectCollectionContainer()->GetCollectionsContainingObject(...) to disambiguate which collection container to use.")
 	virtual void GetCollectionsContainingObject(const FSoftObjectPath& ObjectPath, TArray<FCollectionNameType>& OutCollections, ECollectionRecursionFlags::Flags RecursionMode = ECollectionRecursionFlags::Self) const = 0;
 
 	/** Returns a list of collections in which any of the specified objects exist */
-	UE_DEPRECATED(5.6, "ICollectionContainer has been introduced as a new top level object for housing collections, call GetProjectCollectionContainer to disambiguate which collection container to use.")
+	UE_DEPRECATED(5.6, "ICollectionContainer has been introduced as a new top level object for housing collections, substitute wtih GetProjectCollectionContainer()->GetCollectionsContainingObjects(...) to disambiguate which collection container to use.")
 	virtual void GetCollectionsContainingObjects(const TArray<FSoftObjectPath>& ObjectPaths, TMap<FCollectionNameType, TArray<FSoftObjectPath>>& OutCollectionsAndMatchedObjects, ECollectionRecursionFlags::Flags RecursionMode = ECollectionRecursionFlags::Self) const = 0;
 
 	/** Returns a string containing a comma separated list of collections in which the specified object exists of the specified share type */
-	UE_DEPRECATED(5.6, "ICollectionContainer has been introduced as a new top level object for housing collections, call GetProjectCollectionContainer to disambiguate which collection container to use.")
+	UE_DEPRECATED(5.6, "ICollectionContainer has been introduced as a new top level object for housing collections, substitute wtih GetProjectCollectionContainer()->GetCollectionsStringForObject(...) to disambiguate which collection container to use.")
 	virtual FString GetCollectionsStringForObject(const FSoftObjectPath& ObjectPath, ECollectionShareType::Type ShareType, ECollectionRecursionFlags::Flags RecursionMode = ECollectionRecursionFlags::Self, bool bFullPaths = true) const = 0;
 
 	/** 
@@ -130,7 +130,7 @@ public:
 	 * 
 	 * @warning This function is threadsafe but it does not secure the name from another thread creating the collection in between calls to the ICollectionManager API.
 	 */
-	UE_DEPRECATED(5.6, "ICollectionContainer has been introduced as a new top level object for housing collections, call GetProjectCollectionContainer to disambiguate which collection container to use.")
+	UE_DEPRECATED(5.6, "ICollectionContainer has been introduced as a new top level object for housing collections, substitute wtih GetProjectCollectionContainer()->CreateUniqueCollectionName(...) to disambiguate which collection container to use.")
 	virtual void CreateUniqueCollectionName(FName BaseName, ECollectionShareType::Type ShareType, FName& OutCollectionName) const = 0;
 
 	/**
@@ -142,7 +142,7 @@ public:
 	 * @return true If a collection can be created with the name.
 	 * @return false If the name contains invalid characters or conflicts with an existing collection of the same type.
 	 */
-	UE_DEPRECATED(5.6, "ICollectionContainer has been introduced as a new top level object for housing collections, call GetProjectCollectionContainer to disambiguate which collection container to use.")
+	UE_DEPRECATED(5.6, "ICollectionContainer has been introduced as a new top level object for housing collections, substitute wtih GetProjectCollectionContainer()->IsValidCollectionName(...) to disambiguate which collection container to use.")
 	virtual bool IsValidCollectionName(
 		const FString& CollectionName, ECollectionShareType::Type ShareType, FText* OutError = nullptr) const = 0;
 
@@ -155,7 +155,7 @@ public:
 	 * @param OutError Optional human readable explanation of why the operation failed
 	 * @return true if the add was successful.
 	 */
-	UE_DEPRECATED(5.6, "ICollectionContainer has been introduced as a new top level object for housing collections, call GetProjectCollectionContainer to disambiguate which collection container to use.")
+	UE_DEPRECATED(5.6, "ICollectionContainer has been introduced as a new top level object for housing collections, substitute wtih GetProjectCollectionContainer()->CreateCollection(...) to disambiguate which collection container to use.")
 	virtual bool CreateCollection(
 		FName CollectionName, ECollectionShareType::Type ShareType, ECollectionStorageMode::Type StorageMode,
 		FText* OutError = nullptr) = 0;
@@ -170,7 +170,7 @@ public:
 	 * @param OutError Optional human readable explanation of why the operation failed
 	 * @return true if the rename was successful.
 	 */
-	UE_DEPRECATED(5.6, "ICollectionContainer has been introduced as a new top level object for housing collections, call GetProjectCollectionContainer to disambiguate which collection container to use.")
+	UE_DEPRECATED(5.6, "ICollectionContainer has been introduced as a new top level object for housing collections, substitute wtih GetProjectCollectionContainer()->RenameCollection(...) to disambiguate which collection container to use.")
 	virtual bool RenameCollection(
 		FName CurrentCollectionName, ECollectionShareType::Type CurrentShareType, FName NewCollectionName,
 		ECollectionShareType::Type NewShareType, FText* OutError = nullptr) = 0;
@@ -186,7 +186,7 @@ public:
 	 * @param OutError Optional human readable explanation of why the operation failed
 	 * @return true if the re-parent was successful.
 	 */
-	UE_DEPRECATED(5.6, "ICollectionContainer has been introduced as a new top level object for housing collections, call GetProjectCollectionContainer to disambiguate which collection container to use.")
+	UE_DEPRECATED(5.6, "ICollectionContainer has been introduced as a new top level object for housing collections, substitute wtih GetProjectCollectionContainer()->ReparentCollection(...) to disambiguate which collection container to use.")
 	virtual bool ReparentCollection(
 		FName CollectionName, ECollectionShareType::Type ShareType, FName ParentCollectionName,
 		ECollectionShareType::Type ParentShareType, FText* OutError = nullptr) = 0;
@@ -199,7 +199,7 @@ public:
 	 * @param OutError Optional human readable explanation of why the operation failed
 	 * @return true if the remove was successful
 	 */
-	UE_DEPRECATED(5.6, "ICollectionContainer has been introduced as a new top level object for housing collections, call GetProjectCollectionContainer to disambiguate which collection container to use.")
+	UE_DEPRECATED(5.6, "ICollectionContainer has been introduced as a new top level object for housing collections, substitute wtih GetProjectCollectionContainer()->DestroyCollection(...) to disambiguate which collection container to use.")
 	virtual bool DestroyCollection(FName CollectionName, ECollectionShareType::Type ShareType, FText* OutError = nullptr) = 0;
 
 	/**
@@ -212,11 +212,11 @@ public:
 	 * @param OutError Optional human readable explanation of why the operation failed
 	 * @return true if the add was successful.
 	 */
-	UE_DEPRECATED(5.6, "ICollectionContainer has been introduced as a new top level object for housing collections, call GetProjectCollectionContainer to disambiguate which collection container to use.")
+	UE_DEPRECATED(5.6, "ICollectionContainer has been introduced as a new top level object for housing collections, substitute wtih GetProjectCollectionContainer()->AddToCollection(...) to disambiguate which collection container to use.")
 	virtual bool AddToCollection(
 		FName CollectionName, ECollectionShareType::Type ShareType, const FSoftObjectPath& ObjectPath,
 		FText* OutError = nullptr) = 0;
-	UE_DEPRECATED(5.6, "ICollectionContainer has been introduced as a new top level object for housing collections, call GetProjectCollectionContainer to disambiguate which collection container to use.")
+	UE_DEPRECATED(5.6, "ICollectionContainer has been introduced as a new top level object for housing collections, substitute wtih GetProjectCollectionContainer()->AddToCollection(...) to disambiguate which collection container to use.")
 	virtual bool AddToCollection(
 		FName CollectionName, ECollectionShareType::Type ShareType, TConstArrayView<FSoftObjectPath> ObjectPaths,
 		int32* OutNumAdded = nullptr, FText* OutError = nullptr) = 0;
@@ -231,11 +231,11 @@ public:
 	 * @param OutError Optional human readable explanation of why the operation failed
 	 * @return true if the remove was successful.
 	 */
-	UE_DEPRECATED(5.6, "ICollectionContainer has been introduced as a new top level object for housing collections, call GetProjectCollectionContainer to disambiguate which collection container to use.")
+	UE_DEPRECATED(5.6, "ICollectionContainer has been introduced as a new top level object for housing collections, substitute wtih GetProjectCollectionContainer()->RemoveFromCollection(...) to disambiguate which collection container to use.")
 	virtual bool RemoveFromCollection(
 		FName CollectionName, ECollectionShareType::Type ShareType, const FSoftObjectPath& ObjectPath,
 		FText* OutError = nullptr) = 0;
-	UE_DEPRECATED(5.6, "ICollectionContainer has been introduced as a new top level object for housing collections, call GetProjectCollectionContainer to disambiguate which collection container to use.")
+	UE_DEPRECATED(5.6, "ICollectionContainer has been introduced as a new top level object for housing collections, substitute wtih GetProjectCollectionContainer()->RemoveFromCollection(...) to disambiguate which collection container to use.")
 	virtual bool RemoveFromCollection(
 		FName CollectionName, ECollectionShareType::Type ShareType, TConstArrayView<FSoftObjectPath> ObjectPaths,
 		int32* OutNumRemoved = nullptr, FText* OutError = nullptr) = 0;
@@ -249,7 +249,7 @@ public:
 	 * @param OutError Optional human readable explanation of why the operation failed
 	 * @return true if the set was successful.
 	 */
-	UE_DEPRECATED(5.6, "ICollectionContainer has been introduced as a new top level object for housing collections, call GetProjectCollectionContainer to disambiguate which collection container to use.")
+	UE_DEPRECATED(5.6, "ICollectionContainer has been introduced as a new top level object for housing collections, substitute wtih GetProjectCollectionContainer()->SetDynamicQueryText(...) to disambiguate which collection container to use.")
 	virtual bool SetDynamicQueryText(
 		FName CollectionName, ECollectionShareType::Type ShareType, const FString& InQueryText,
 		FText* OutError = nullptr) = 0;
@@ -263,7 +263,7 @@ public:
 	 * @param OutError Optional human readable explanation of why the operation failed
 	 * @return true if the get was successful.
 	 */
-	UE_DEPRECATED(5.6, "ICollectionContainer has been introduced as a new top level object for housing collections, call GetProjectCollectionContainer to disambiguate which collection container to use.")
+	UE_DEPRECATED(5.6, "ICollectionContainer has been introduced as a new top level object for housing collections, substitute wtih GetProjectCollectionContainer()->GetDynamicQueryText(...) to disambiguate which collection container to use.")
 	virtual bool GetDynamicQueryText(
 		FName CollectionName, ECollectionShareType::Type ShareType, FString& OutQueryText,
 		FText* OutError = nullptr) const = 0;
@@ -278,7 +278,7 @@ public:
 	 * @param OutError Optional human readable explanation of why the operation failed
 	 * @return true if the get was successful.
 	 */
-	UE_DEPRECATED(5.6, "ICollectionContainer has been introduced as a new top level object for housing collections, call GetProjectCollectionContainer to disambiguate which collection container to use.")
+	UE_DEPRECATED(5.6, "ICollectionContainer has been introduced as a new top level object for housing collections, substitute wtih GetProjectCollectionContainer()->TestDynamicQuery(...) to disambiguate which collection container to use.")
 	virtual bool TestDynamicQuery(
 		FName CollectionName, ECollectionShareType::Type ShareType, const ITextFilterExpressionContext& InContext,
 		bool& OutResult, FText* OutError = nullptr) const = 0;
@@ -291,7 +291,7 @@ public:
 	 * @param OutError Optional human readable explanation of why the operation failed
 	 * @return true if the clear was successful.
 	 */
-	UE_DEPRECATED(5.6, "ICollectionContainer has been introduced as a new top level object for housing collections, call GetProjectCollectionContainer to disambiguate which collection container to use.")
+	UE_DEPRECATED(5.6, "ICollectionContainer has been introduced as a new top level object for housing collections, substitute wtih GetProjectCollectionContainer()->EmptyCollection(...) to disambiguate which collection container to use.")
 	virtual bool EmptyCollection(FName CollectionName, ECollectionShareType::Type ShareType, FText* OutError = nullptr) = 0;
 
 	/**
@@ -307,7 +307,7 @@ public:
 	 * @param OutError Optional human readable explanation of why the operation failed
 	 * @return true if the save was successful.
 	 */
-	UE_DEPRECATED(5.6, "ICollectionContainer has been introduced as a new top level object for housing collections, call GetProjectCollectionContainer to disambiguate which collection container to use.")
+	UE_DEPRECATED(5.6, "ICollectionContainer has been introduced as a new top level object for housing collections, substitute wtih GetProjectCollectionContainer()->SaveCollection(...) to disambiguate which collection container to use.")
 	virtual bool SaveCollection(FName CollectionName, ECollectionShareType::Type ShareType, FText* OutError = nullptr) = 0;
 
 	/**
@@ -322,7 +322,7 @@ public:
 	 * @param OutError Optional human readable explanation of why the operation failed
 	 * @return true if the update was successful.
 	 */
-	UE_DEPRECATED(5.6, "ICollectionContainer has been introduced as a new top level object for housing collections, call GetProjectCollectionContainer to disambiguate which collection container to use.")
+	UE_DEPRECATED(5.6, "ICollectionContainer has been introduced as a new top level object for housing collections, substitute wtih GetProjectCollectionContainer()->UpdateCollection(...) to disambiguate which collection container to use.")
 	virtual bool UpdateCollection(FName CollectionName, ECollectionShareType::Type ShareType, FText* OutError = nullptr) = 0;
 
 	/**
@@ -334,7 +334,7 @@ public:
 	 * @param OutError Optional human readable explanation of why the operation failed
 	 * @return true if the status info was filled in.
 	 */
-	UE_DEPRECATED(5.6, "ICollectionContainer has been introduced as a new top level object for housing collections, call GetProjectCollectionContainer to disambiguate which collection container to use.")
+	UE_DEPRECATED(5.6, "ICollectionContainer has been introduced as a new top level object for housing collections, substitute wtih GetProjectCollectionContainer()->GetCollectionStatusInfo(...) to disambiguate which collection container to use.")
 	virtual bool GetCollectionStatusInfo(
 		FName CollectionName, ECollectionShareType::Type ShareType, FCollectionStatusInfo& OutStatusInfo,
 		FText* OutError = nullptr) const = 0;
@@ -345,7 +345,7 @@ public:
 	 * @param OutColors The optional color array to populate.
 	 * @return true if there are any custom colors in use, false otherwise.
 	 */
-	UE_DEPRECATED(5.6, "ICollectionContainer has been introduced as a new top level object for housing collections, call GetProjectCollectionContainer to disambiguate which collection container to use.")
+	UE_DEPRECATED(5.6, "ICollectionContainer has been introduced as a new top level object for housing collections, substitute wtih GetProjectCollectionContainer()->HasCollectionColors(...) to disambiguate which collection container to use.")
 	virtual bool HasCollectionColors(TArray<FLinearColor>* OutColors = nullptr) const = 0;
 
 	/**
@@ -357,7 +357,7 @@ public:
 	 * @param OutError Optional human readable explanation of why the operation failed
 	 * @return true if the collection was found, and the color was optionally filled in.
 	 */
-	UE_DEPRECATED(5.6, "ICollectionContainer has been introduced as a new top level object for housing collections, call GetProjectCollectionContainer to disambiguate which collection container to use.")
+	UE_DEPRECATED(5.6, "ICollectionContainer has been introduced as a new top level object for housing collections, substitute wtih GetProjectCollectionContainer()->GetCollectionColor(...) to disambiguate which collection container to use.")
 	virtual bool GetCollectionColor(
 		FName CollectionName, ECollectionShareType::Type ShareType, TOptional<FLinearColor>& OutColor,
 		FText* OutError = nullptr) const = 0;
@@ -371,7 +371,7 @@ public:
 	 * @param OutError Optional human readable explanation of why the operation failed
 	 * @return true if the collection was found, and the color was set.
 	 */
-	UE_DEPRECATED(5.6, "ICollectionContainer has been introduced as a new top level object for housing collections, call GetProjectCollectionContainer to disambiguate which collection container to use.")
+	UE_DEPRECATED(5.6, "ICollectionContainer has been introduced as a new top level object for housing collections, substitute wtih GetProjectCollectionContainer()->SetCollectionColor(...) to disambiguate which collection container to use.")
 	virtual bool SetCollectionColor(
 		FName CollectionName, ECollectionShareType::Type ShareType, const TOptional<FLinearColor>& NewColor,
 		FText* OutError = nullptr) = 0;
@@ -385,7 +385,7 @@ public:
 	 * @param OutError Optional human readable explanation of why the operation failed
 	 * @return true if the status info was filled in.
 	 */
-	UE_DEPRECATED(5.6, "ICollectionContainer has been introduced as a new top level object for housing collections, call GetProjectCollectionContainer to disambiguate which collection container to use.")
+	UE_DEPRECATED(5.6, "ICollectionContainer has been introduced as a new top level object for housing collections, substitute wtih GetProjectCollectionContainer()->GetCollectionStorageMode(...) to disambiguate which collection container to use.")
 	virtual bool GetCollectionStorageMode(
 		FName CollectionName, ECollectionShareType::Type ShareType, ECollectionStorageMode::Type& OutStorageMode,
 		FText* OutError = nullptr) const = 0;
@@ -399,7 +399,7 @@ public:
 	 * @param OutError Optional human readable explanation of why the operation failed
 	 * @return true if the object is in the collection.
 	 */
-	UE_DEPRECATED(5.6, "ICollectionContainer has been introduced as a new top level object for housing collections, call GetProjectCollectionContainer to disambiguate which collection container to use.")
+	UE_DEPRECATED(5.6, "ICollectionContainer has been introduced as a new top level object for housing collections, substitute wtih GetProjectCollectionContainer()->IsObjectInCollection(...) to disambiguate which collection container to use.")
 	virtual bool IsObjectInCollection(
 		const FSoftObjectPath& ObjectPath, FName CollectionName, ECollectionShareType::Type ShareType,
 		ECollectionRecursionFlags::Flags RecursionMode = ECollectionRecursionFlags::Self,
@@ -416,7 +416,7 @@ public:
 	 * @param OutError Optional human readable explanation of why the operation failed
 	 * @return true if the parent is considered valid for this collection.
 	 */
-	UE_DEPRECATED(5.6, "ICollectionContainer has been introduced as a new top level object for housing collections, call GetProjectCollectionContainer to disambiguate which collection container to use.")
+	UE_DEPRECATED(5.6, "ICollectionContainer has been introduced as a new top level object for housing collections, substitute wtih GetProjectCollectionContainer()->IsValidParentCollection(...) to disambiguate which collection container to use.")
 	virtual bool IsValidParentCollection(
 		FName CollectionName, ECollectionShareType::Type ShareType, FName ParentCollectionName,
 		ECollectionShareType::Type ParentShareType, FText* OutError = nullptr) const = 0;
@@ -462,37 +462,37 @@ public:
 
 	/** Event for when collections are created */
 	DECLARE_EVENT_OneParam( ICollectionManager, FCollectionCreatedEvent, const FCollectionNameType& );
-	UE_DEPRECATED(5.6, "ICollectionContainer has been introduced as a new top level object for housing collections, call GetProjectCollectionContainer to disambiguate which collection container to use.")
+	UE_DEPRECATED(5.6, "ICollectionContainer has been introduced as a new top level object for housing collections, substitute wtih GetProjectCollectionContainer()->OnCollectionCreated() to disambiguate which collection container to use.")
 	virtual FCollectionCreatedEvent& OnCollectionCreated() = 0;
 
 	/** Event for when collections are destroyed */
 	DECLARE_EVENT_OneParam( ICollectionManager, FCollectionDestroyedEvent, const FCollectionNameType& );
-	UE_DEPRECATED(5.6, "ICollectionContainer has been introduced as a new top level object for housing collections, call GetProjectCollectionContainer to disambiguate which collection container to use.")
+	UE_DEPRECATED(5.6, "ICollectionContainer has been introduced as a new top level object for housing collections, substitute wtih GetProjectCollectionContainer()->OnCollectionDestroyed() to disambiguate which collection container to use.")
 	virtual FCollectionDestroyedEvent& OnCollectionDestroyed() = 0;
 
 	/** Event for when assets are added to a collection */
 	DECLARE_MULTICAST_DELEGATE_TwoParams(FOnAssetsAddedToCollection, const FCollectionNameType&, TConstArrayView<FSoftObjectPath> );
-	UE_DEPRECATED(5.6, "ICollectionContainer has been introduced as a new top level object for housing collections, call GetProjectCollectionContainer to disambiguate which collection container to use.")
+	UE_DEPRECATED(5.6, "ICollectionContainer has been introduced as a new top level object for housing collections, substitute wtih GetProjectCollectionContainer()->OnAssetsAddedToCollection() to disambiguate which collection container to use.")
 	virtual FOnAssetsAddedToCollection& OnAssetsAddedToCollection() = 0;
 
 	/** Event for when assets are removed from a collection */
 	DECLARE_MULTICAST_DELEGATE_TwoParams( FOnAssetsRemovedFromCollection, const FCollectionNameType&, TConstArrayView<FSoftObjectPath> );
-	UE_DEPRECATED(5.6, "ICollectionContainer has been introduced as a new top level object for housing collections, call GetProjectCollectionContainer to disambiguate which collection container to use.")
+	UE_DEPRECATED(5.6, "ICollectionContainer has been introduced as a new top level object for housing collections, substitute wtih GetProjectCollectionContainer()->OnAssetsRemovedFromCollection() to disambiguate which collection container to use.")
 	virtual FOnAssetsRemovedFromCollection& OnAssetsRemovedFromCollection() = 0;
 
 	/** Event for when collections are renamed */
 	DECLARE_EVENT_TwoParams( ICollectionManager, FCollectionRenamedEvent, const FCollectionNameType&, const FCollectionNameType& );
-	UE_DEPRECATED(5.6, "ICollectionContainer has been introduced as a new top level object for housing collections, call GetProjectCollectionContainer to disambiguate which collection container to use.")
+	UE_DEPRECATED(5.6, "ICollectionContainer has been introduced as a new top level object for housing collections, substitute wtih GetProjectCollectionContainer()->OnCollectionRenamed() to disambiguate which collection container to use.")
 	virtual FCollectionRenamedEvent& OnCollectionRenamed() = 0;
 
 	/** Event for when collections are re-parented (params: Collection, OldParent, NewParent) */
 	DECLARE_EVENT_ThreeParams( ICollectionManager, FCollectionReparentedEvent, const FCollectionNameType&, const TOptional<FCollectionNameType>&, const TOptional<FCollectionNameType>& );
-	UE_DEPRECATED(5.6, "ICollectionContainer has been introduced as a new top level object for housing collections, call GetProjectCollectionContainer to disambiguate which collection container to use.")
+	UE_DEPRECATED(5.6, "ICollectionContainer has been introduced as a new top level object for housing collections, substitute wtih GetProjectCollectionContainer()->OnCollectionReparented() to disambiguate which collection container to use.")
 	virtual FCollectionReparentedEvent& OnCollectionReparented() = 0;
 
 	/** Event for when collections is updated, or otherwise changed and we can't tell exactly how (eg, after updating from source control and merging) */
 	DECLARE_EVENT_OneParam( ICollectionManager, FCollectionUpdatedEvent, const FCollectionNameType& );
-	UE_DEPRECATED(5.6, "ICollectionContainer has been introduced as a new top level object for housing collections, call GetProjectCollectionContainer to disambiguate which collection container to use.")
+	UE_DEPRECATED(5.6, "ICollectionContainer has been introduced as a new top level object for housing collections, substitute wtih GetProjectCollectionContainer()->OnCollectionUpdated() to disambiguate which collection container to use.")
 	virtual FCollectionUpdatedEvent& OnCollectionUpdated() = 0;
 
 	/** 
