Add enum range declaration for ECollectionShareType

diff --git a/Engine/Source/Developer/CollectionManager/Public/CollectionManagerTypes.h b/Engine/Source/Developer/CollectionManager/Public/CollectionManagerTypes.h
--- a/Engine/Source/Developer/CollectionManager/Public/CollectionManagerTypes.h
+++ b/Engine/Source/Developer/CollectionManager/Public/CollectionManagerTypes.h
@@ -3,6 +3,7 @@
 #pragma once
 
 #include "CoreMinimal.h"
+#include "Misc/EnumRange.h"
 #include "UObject/SoftObjectPath.h"
 #include "ISourceControlState.h"
 
@@ -190,6 +191,8 @@ struct ECollectionShareType
 	}
 };
 
+ENUM_RANGE_BY_COUNT(ECollectionShareType::Type, ECollectionShareType::CST_All);
+
 /** Controls how the collections manager will recurse when performing work against a given collection */
 struct ECollectionRecursionFlags
 {
