Add concept of auto-generated collections to allow fixing redirectors on developer machines without modifying certain collections, as those will be updated by the build system anyway

diff --git a/Engine/Source/Developer/CollectionManager/Private/Collection.cpp b/Engine/Source/Developer/CollectionManager/Private/Collection.cpp
--- a/Engine/Source/Developer/CollectionManager/Private/Collection.cpp
+++ b/Engine/Source/Developer/CollectionManager/Private/Collection.cpp
@@ -710,6 +710,7 @@ void FCollection::SaveHeaderPairs(TMap<FString,FString>& OutHeaderPairs) const
 	OutHeaderPairs.Add(TEXT("Type"), ECollectionStorageMode::ToString(StorageMode));
 	OutHeaderPairs.Add(TEXT("Guid"), CollectionGuid.ToString(EGuidFormats::DigitsWithHyphens));
 	OutHeaderPairs.Add(TEXT("ParentGuid"), ParentCollectionGuid.ToString(EGuidFormats::DigitsWithHyphens));
+	OutHeaderPairs.Add(TEXT("IsAutoGenerated"), bIsAutoGenerated ? TEXT("true") : TEXT("false"));
 	if (CollectionColor)
 	{
 		OutHeaderPairs.Add(TEXT("Color"), CollectionColor->ToString());
@@ -765,6 +766,12 @@ bool FCollection::LoadHeaderPairs(const TMap<FString,FString>& InHeaderPairs)
 		}
 	}
 
+	bIsAutoGenerated = false;
+	if (const FString* IsAutoGeneratedString = InHeaderPairs.Find(TEXT("IsAutoGenerated")))
+	{
+		bIsAutoGenerated = IsAutoGeneratedString->ToBool();
+	}
+
 	return FileVersion > 0 && FileVersion <= ECollectionVersion::CurrentVersion;
 }
 
diff --git a/Engine/Source/Developer/CollectionManager/Private/Collection.h b/Engine/Source/Developer/CollectionManager/Private/Collection.h
--- a/Engine/Source/Developer/CollectionManager/Private/Collection.h
+++ b/Engine/Source/Developer/CollectionManager/Private/Collection.h
@@ -119,6 +119,9 @@ public:
 	/** Set the GUID of the collection we are parented under */
 	FORCEINLINE void SetParentCollectionGuid(const FGuid& NewGuid) { ParentCollectionGuid = NewGuid; }
 
+	/** Get whether this collection is auto-generated by something external instead of maintained by users */
+	FORCEINLINE bool IsAutoGenerated() const { return bIsAutoGenerated; }
+
 	/** Get the color of the collection (if any) */
 	FORCEINLINE TOptional<FLinearColor> GetCollectionColor() const { return CollectionColor; }
 
@@ -210,6 +213,9 @@ private:
 	/** Indicates if the collection has changes in memory that weren't saved to disk. */
 	bool bChangedSinceLastDiskSnapshot = false;
 
+	/** Indicates if the collection is auto-generated by something external instead of maintained by users */
+	bool bIsAutoGenerated = false;
+
 	/** The filename used to load this collection. Empty if it is new or never loaded from disk. */
 	FString SourceFilename;
 
diff --git a/Engine/Source/Developer/CollectionManager/Private/CollectionContainer.cpp b/Engine/Source/Developer/CollectionManager/Private/CollectionContainer.cpp
--- a/Engine/Source/Developer/CollectionManager/Private/CollectionContainer.cpp
+++ b/Engine/Source/Developer/CollectionManager/Private/CollectionContainer.cpp
@@ -2237,7 +2237,7 @@ bool FCollectionContainer::HandleRedirectorsDeleted(TConstArrayView<FSoftObjectP
 				const FCollectionNameType& CollectionKey = AvailableCollection.Key;
 				const TSharedRef<FCollection>& Collection = AvailableCollection.Value;
 
-				if (Collection->IsRedirectorInCollection(ObjectPath))
+				if (Collection->IsRedirectorInCollection(ObjectPath) && !Collection->IsAutoGenerated())
 				{
 					CollectionsToSave.Add(CollectionKey);
 				}
